Experimental plugin for remapping animation weight curves at runtime using simple expressions.

diff --git a/Engine/Plugins/Experimental/Animation/CurveExpression/CurveExpression.uplugin b/Engine/Plugins/Experimental/Animation/CurveExpression/CurveExpression.uplugin
new file mode 100644
--- /dev/null
+++ b/Engine/Plugins/Experimental/Animation/CurveExpression/CurveExpression.uplugin
@@ -0,0 +1,30 @@
+{
+	"FileVersion": 3,
+	"Version": 1,
+	"VersionName": "0.1",
+	"FriendlyName": "Animation Curve Expressions",
+	"Description": "Experimental Curve Remapper using Simple Math Expressions",
+	"Category": "Animation",
+	"CreatedBy": "Epic Games, Inc.",
+	"CreatedByURL": "http://epicgames.com",
+	"DocsURL": "",
+	"MarketplaceURL": "",
+	"SupportURL": "",
+	"EnabledByDefault": false,
+	"CanContainContent": false,
+	"IsBetaVersion": false,
+	"IsExperimentalVersion": true,
+	"Installed": false,
+	"Modules": [
+		{
+			"Name": "CurveExpression",
+			"Type": "Runtime",
+			"LoadingPhase": "PreDefault"
+		},
+		{
+			"Name": "CurveExpressionEditor",
+			"Type": "UncookedOnly",
+			"LoadingPhase": "PreDefault"
+		}
+	]
+}
diff --git a/Engine/Plugins/Experimental/Animation/CurveExpression/Source/Editor/CurveExpressionEditor.Build.cs b/Engine/Plugins/Experimental/Animation/CurveExpression/Source/Editor/CurveExpressionEditor.Build.cs
new file mode 100644
--- /dev/null
+++ b/Engine/Plugins/Experimental/Animation/CurveExpression/Source/Editor/CurveExpressionEditor.Build.cs
@@ -0,0 +1,37 @@
+// Copyright Epic Games, Inc. All Rights Reserved.
+
+using UnrealBuildTool;
+
+public class CurveExpressionEditor : ModuleRules
+{
+	public CurveExpressionEditor(ReadOnlyTargetRules Target) : base(Target)
+	{
+		PublicDependencyModuleNames.AddRange(
+			new string[]
+			{
+				"AnimGraph",
+				"AnimationCore",
+				"Core",
+				"CoreUObject",
+				"CurveExpression",
+				"Engine",
+				
+				// UI 
+				"PropertyEditor",
+				"Slate",
+				"SlateCore",
+				"RewindDebuggerInterface",
+				"EditorStyle",
+				"UnrealEd",
+				"InputCore"
+			}
+		);
+
+		PrivateDependencyModuleNames.AddRange(
+			new string[]
+			{
+				"BlueprintGraph"
+			}
+		);
+	}
+}
diff --git a/Engine/Plugins/Experimental/Animation/CurveExpression/Source/Editor/Private/AnimGraphNode_RemapCurvesFromMesh.cpp b/Engine/Plugins/Experimental/Animation/CurveExpression/Source/Editor/Private/AnimGraphNode_RemapCurvesFromMesh.cpp
new file mode 100644
--- /dev/null
+++ b/Engine/Plugins/Experimental/Animation/CurveExpression/Source/Editor/Private/AnimGraphNode_RemapCurvesFromMesh.cpp
@@ -0,0 +1,23 @@
+// Copyright Epic Games, Inc. All Rights Reserved.
+
+#include "AnimGraphNode_RemapCurvesFromMesh.h"
+
+#define LOCTEXT_NAMESPACE "AnimGraphNode_RemapCurvesFromMesh"
+
+
+FText UAnimGraphNode_RemapCurvesFromMesh::GetNodeTitle(ENodeTitleType::Type TitleType) const
+{
+	return LOCTEXT("NodeTitle", "Remap Curves From Mesh");
+}
+
+FText UAnimGraphNode_RemapCurvesFromMesh::GetTooltipText() const
+{
+	return LOCTEXT("Tooltip", "The Remap Curves From Mesh node copies curves from another component to this. Can be used to map any curve to any and perform mathematical operations on them.");
+}
+
+void UAnimGraphNode_RemapCurvesFromMesh::GetOutputLinkAttributes(FNodeAttributeArray& OutAttributes) const
+{
+	OutAttributes.Add(UE::Anim::FAttributes::Curves);
+}
+
+#undef LOCTEXT_NAMESPACE
diff --git a/Engine/Plugins/Experimental/Animation/CurveExpression/Source/Editor/Public/AnimGraphNode_RemapCurvesFromMesh.h b/Engine/Plugins/Experimental/Animation/CurveExpression/Source/Editor/Public/AnimGraphNode_RemapCurvesFromMesh.h
new file mode 100644
--- /dev/null
+++ b/Engine/Plugins/Experimental/Animation/CurveExpression/Source/Editor/Public/AnimGraphNode_RemapCurvesFromMesh.h
@@ -0,0 +1,30 @@
+// Copyright Epic Games, Inc. All Rights Reserved.
+
+#pragma once
+
+#include "AnimGraphNode_Base.h"
+
+#include "AnimNode_RemapCurvesFromMesh.h"
+
+#include "AnimGraphNode_RemapCurvesFromMesh.generated.h"
+
+
+UCLASS(MinimalAPI)
+class UAnimGraphNode_RemapCurvesFromMesh :
+	public UAnimGraphNode_Base
+{
+	GENERATED_BODY()
+
+	UPROPERTY(EditAnywhere, Category=Settings)
+	FAnimNode_RemapCurvesFromMesh Node;
+
+	// UEdGraphNode interface
+	virtual FText GetNodeTitle(ENodeTitleType::Type TitleType) const override;
+	virtual FText GetTooltipText() const override;
+	// End of UEdGraphNode interface
+
+	// UAnimGraphNode_Base interface
+	virtual bool UsingCopyPoseFromMesh() const override { return true; }
+	virtual void GetOutputLinkAttributes(FNodeAttributeArray& OutAttributes) const override;
+	// End of UAnimGraphNode_Base interface
+};
diff --git a/Engine/Plugins/Experimental/Animation/CurveExpression/Source/Runtime/CurveExpression.Build.cs b/Engine/Plugins/Experimental/Animation/CurveExpression/Source/Runtime/CurveExpression.Build.cs
new file mode 100644
--- /dev/null
+++ b/Engine/Plugins/Experimental/Animation/CurveExpression/Source/Runtime/CurveExpression.Build.cs
@@ -0,0 +1,25 @@
+// Copyright Epic Games, Inc. All Rights Reserved.
+
+using UnrealBuildTool;
+
+public class CurveExpression : ModuleRules
+{
+	public CurveExpression(ReadOnlyTargetRules Target) : base(Target)
+	{
+		PublicDependencyModuleNames.AddRange(
+			new string[]
+			{
+				"Core",
+				"CoreUObject",
+				"Engine",
+				"AnimationCore",
+			}
+		);
+
+		PrivateDependencyModuleNames.AddRange(
+			new string[]
+            {
+			}
+		);
+	}
+}
diff --git a/Engine/Plugins/Experimental/Animation/CurveExpression/Source/Runtime/Private/AnimNode_RemapCurvesFromMesh.cpp b/Engine/Plugins/Experimental/Animation/CurveExpression/Source/Runtime/Private/AnimNode_RemapCurvesFromMesh.cpp
new file mode 100644
--- /dev/null
+++ b/Engine/Plugins/Experimental/Animation/CurveExpression/Source/Runtime/Private/AnimNode_RemapCurvesFromMesh.cpp
@@ -0,0 +1,288 @@
+// Copyright Epic Games, Inc. All Rights Reserved.
+
+#include "AnimNode_RemapCurvesFromMesh.h"
+
+#include "ExpressionEvaluator.h"
+
+void FAnimNode_RemapCurvesFromMesh::Initialize_AnyThread(
+	const FAnimationInitializeContext& Context
+	)
+{
+	DECLARE_SCOPE_HIERARCHICAL_COUNTER_ANIMNODE(Initialize_AnyThread)
+	Super::Initialize_AnyThread(Context);
+	SourcePose.Initialize(Context);
+}
+
+
+void FAnimNode_RemapCurvesFromMesh::CacheBones_AnyThread(
+	const FAnimationCacheBonesContext& Context
+	)
+{
+	DECLARE_SCOPE_HIERARCHICAL_COUNTER_ANIMNODE(CacheBones_AnyThread)
+	Super::CacheBones_AnyThread(Context);
+	SourcePose.CacheBones(Context);
+}
+	
+
+void FAnimNode_RemapCurvesFromMesh::Update_AnyThread(
+	const FAnimationUpdateContext& Context
+	)
+{
+	DECLARE_SCOPE_HIERARCHICAL_COUNTER_ANIMNODE(Update_AnyThread)
+
+	// Run update on input pose nodes
+	SourcePose.Update(Context);
+
+	// Evaluate any BP logic plugged into this node
+	GetEvaluateGraphExposedInputs().Execute(Context);
+	
+	if (bExpressionsImmutable && ExpressionEngine.IsSet() && !CurveExpressions.IsEmpty() && CachedExpressions.IsEmpty())
+	{
+		CachedExpressions.Reset();
+		for (const TPair<FName, FString>& ExpressionPair: CurveExpressions)
+		{
+			if (CurveNameToUIDMap.Contains(ExpressionPair.Key))
+			{
+				using namespace CurveExpression::Evaluator;
+
+				TVariant<FExpressionObject, FParseError> Result = ExpressionEngine->Parse(ExpressionPair.Value);
+				if (const FExpressionObject* Expression = Result.TryGet<FExpressionObject>())
+				{
+					CachedExpressions.Add(ExpressionPair.Key, *Expression);
+				}
+				// Parse errors we deal with in the animgraph node.
+			}
+		}
+	}			
+}
+
+
+void FAnimNode_RemapCurvesFromMesh::Evaluate_AnyThread(
+	FPoseContext& Output
+	)
+{
+	DECLARE_SCOPE_HIERARCHICAL_COUNTER_ANIMNODE(Evaluate_AnyThread)
+
+	FPoseContext SourceData(Output);
+	SourcePose.Evaluate(SourceData);
+	
+	Output = SourceData;
+	
+	// If we have an expression engine, evaluate the expressions that have a matching target curve.
+	// If the expressions are not immutable between compiles, then we need to reparse them each time. 
+	if (ExpressionEngine.IsSet())
+	{
+		if (bExpressionsImmutable)
+		{
+			for (const TTuple<FName, CurveExpression::Evaluator::FExpressionObject>& ExpressionItem: CachedExpressions)
+			{
+				if (const SmartName::UID_Type* UID = CurveNameToUIDMap.Find(ExpressionItem.Key))
+				{
+					Output.Curve.Set(*UID, ExpressionEngine->Execute(ExpressionItem.Value));
+				}
+			}
+		}
+		else
+		{
+			for (const TTuple<FName, FString>& ExpressionItem: CurveExpressions)
+			{
+				if (const SmartName::UID_Type* UID = CurveNameToUIDMap.Find(ExpressionItem.Key))
+				{
+					TOptional<float> Result = ExpressionEngine->Evaluate(ExpressionItem.Value);
+					if (Result.IsSet())
+					{
+						Output.Curve.Set(*UID, *Result);
+					}
+				}
+			}
+		}
+	}
+}
+
+
+void FAnimNode_RemapCurvesFromMesh::GatherDebugData(
+	FNodeDebugData& DebugData
+	)
+{
+	DECLARE_SCOPE_HIERARCHICAL_COUNTER_ANIMNODE(GatherDebugData)
+	FString DebugLine = DebugData.GetNodeName(this);
+
+	DebugLine += FString::Printf(TEXT("('%s')"), *GetNameSafe(CurrentlyUsedSourceMeshComponent.IsValid() ? CurrentlyUsedSourceMeshComponent.Get()->SkeletalMesh : nullptr));
+	DebugData.AddDebugItem(DebugLine, true);
+}
+
+
+void FAnimNode_RemapCurvesFromMesh::PreUpdate(
+	const UAnimInstance* InAnimInstance
+	)
+{
+	QUICK_SCOPE_CYCLE_COUNTER(FAnimNode_RemapCurvesFromMesh_PreUpdate);
+
+	// Make sure we're using the correct source and target skeleton components, since they
+	// may have changed from underneath us.
+	RefreshMeshComponent(InAnimInstance->GetSkelMeshComponent());
+
+	const USkeletalMeshComponent* CurrentMeshComponent = CurrentlyUsedSourceMeshComponent.IsValid() ? CurrentlyUsedSourceMeshComponent.Get() : nullptr;
+
+	if (CurrentMeshComponent && CurrentMeshComponent->SkeletalMesh && CurrentMeshComponent->IsRegistered())
+	{
+		// If our source is running under master-pose, then get bone data from there
+		if(USkeletalMeshComponent* MasterPoseComponent = Cast<USkeletalMeshComponent>(CurrentMeshComponent->MasterPoseComponent.Get()))
+		{
+			CurrentMeshComponent = MasterPoseComponent;
+		}
+
+		// re-check mesh component validity as it may have changed to master
+		if(CurrentMeshComponent->SkeletalMesh && CurrentMeshComponent->IsRegistered())
+		{
+			if (ExpressionEngine.IsSet())
+			{
+				if (const UAnimInstance* SourceAnimInstance = CurrentMeshComponent->GetAnimInstance())
+				{
+					ExpressionEngine->UpdateConstantValues(SourceAnimInstance->GetAnimationCurveList(EAnimCurveType::AttributeCurve));
+				}
+				else
+				{
+					ExpressionEngine.Reset();
+				}
+			}
+		}
+		else
+		{
+			ExpressionEngine.Reset();
+		}
+	}
+}
+
+
+void FAnimNode_RemapCurvesFromMesh::ReinitializeMeshComponent(
+	USkeletalMeshComponent* InNewSkeletalMeshComponent, 
+	USkeletalMeshComponent* InTargetMeshComponent
+	)
+{
+	CurrentlyUsedSourceMeshComponent.Reset();
+	CurrentlyUsedSourceMesh.Reset();
+	CurrentlyUsedTargetMesh.Reset();
+	
+	ExpressionEngine.Reset();
+	CurveNameToUIDMap.Reset();
+	CachedExpressions.Reset();
+
+	if (InTargetMeshComponent && IsValid(InNewSkeletalMeshComponent) && InNewSkeletalMeshComponent->SkeletalMesh)
+	{
+		USkeletalMesh* SourceSkelMesh = InNewSkeletalMeshComponent->SkeletalMesh;
+		USkeletalMesh* TargetSkelMesh = InTargetMeshComponent->SkeletalMesh;
+		
+		if (IsValid(SourceSkelMesh) && !SourceSkelMesh->HasAnyFlags(RF_NeedPostLoad) &&
+			IsValid(TargetSkelMesh) && !TargetSkelMesh->HasAnyFlags(RF_NeedPostLoad))
+		{
+			CurrentlyUsedSourceMeshComponent = InNewSkeletalMeshComponent;
+			CurrentlyUsedSourceMesh = SourceSkelMesh;
+			CurrentlyUsedTargetMesh = TargetSkelMesh;
+
+			// Grab the source curves and use their names to seed the expression evaluation engine.
+			const USkeleton* SourceSkeleton = SourceSkelMesh->GetSkeleton();
+			if (ensureMsgf(SourceSkeleton, TEXT("Invalid null source skeleton : %s"), *GetNameSafe(TargetSkelMesh)))
+			{
+				const FSmartNameMapping* SourceContainer = SourceSkeleton->GetSmartNameContainer(USkeleton::AnimCurveMappingName);
+
+				TMap<FName, float> SourceCurves;
+				
+				SourceContainer->Iterate([&SourceCurves](const FSmartNameMapping* Mapping, const SmartName::UID_Type ID)
+				{
+					if (FName CurveName; Mapping->GetName(ID, CurveName))
+					{
+						const FCurveMetaData* MetaData = Mapping->GetCurveMetaData(CurveName);
+						FPlatformMisc::LowLevelOutputDebugStringf(TEXT("CURVE[%s] - Morph: %d - Mat: %d\n"), *CurveName.ToString(), MetaData->Type.bMaterial, MetaData->Type.bMorphtarget);
+						SourceCurves.Add(CurveName, 0.0f);
+					}
+				});
+
+				ExpressionEngine.Emplace(MoveTemp(SourceCurves));
+			}
+
+			// Grab the target curves and build the Name -> UID mapping so that we can call
+			// FBlendCurve::Set in Evalute_AnyThread 
+			const USkeleton* TargetSkeleton = TargetSkelMesh->GetSkeleton();
+
+			if (ensureMsgf(TargetSkeleton, TEXT("Invalid null target skeleton : %s"), *GetNameSafe(TargetSkelMesh)))
+			{
+				const FSmartNameMapping* TargetContainer = TargetSkeleton->GetSmartNameContainer(USkeleton::AnimCurveMappingName);
+
+				// CurveExpressions isn't filled in here, since we haven't evaluated BP downstream.
+				// So instead, we just take a full copy of the smart mapping.
+				TargetContainer->Iterate([this](const FSmartNameMapping* Mapping, const SmartName::UID_Type ID)
+				{
+					if (FName CurveName; Mapping->GetName(ID, CurveName))
+					{
+						CurveNameToUIDMap.Add(CurveName, ID);
+					}
+				});
+			}
+		}
+	}
+}
+	
+
+void FAnimNode_RemapCurvesFromMesh::RefreshMeshComponent(
+	USkeletalMeshComponent* InTargetMeshComponent
+	)
+{
+	auto ResetMeshComponent = [this](USkeletalMeshComponent* InMeshComponent, USkeletalMeshComponent* InTargetMeshComponent)
+	{
+		// if current mesh exists, but not same as input mesh
+		if (const USkeletalMeshComponent* CurrentMeshComponent = CurrentlyUsedSourceMeshComponent.Get())
+		{
+			// if component has been changed, reinitialize
+			if (CurrentMeshComponent != InMeshComponent)
+			{
+				ReinitializeMeshComponent(InMeshComponent, InTargetMeshComponent);
+			}
+			// if component is still same but mesh has been changed, we have to reinitialize
+			else if (CurrentMeshComponent->SkeletalMesh != CurrentlyUsedSourceMesh.Get())
+			{
+				ReinitializeMeshComponent(InMeshComponent, InTargetMeshComponent);
+			}
+			else if (InTargetMeshComponent)
+			{
+				// see if target mesh has changed
+				if (InTargetMeshComponent->SkeletalMesh != CurrentlyUsedTargetMesh.Get())
+				{
+					ReinitializeMeshComponent(InMeshComponent, InTargetMeshComponent);
+				}
+			}
+		}
+		// if not valid, but input mesh is
+		else if (!CurrentMeshComponent && InMeshComponent)
+		{
+			ReinitializeMeshComponent(InMeshComponent, InTargetMeshComponent);
+		}
+	};
+
+	if(SourceMeshComponent.IsValid())
+	{
+		ResetMeshComponent(SourceMeshComponent.Get(), InTargetMeshComponent);
+	}
+	else if (bUseAttachedParent)
+	{
+		if (InTargetMeshComponent)
+		{
+			if (USkeletalMeshComponent* ParentComponent = Cast<USkeletalMeshComponent>(InTargetMeshComponent->GetAttachParent()))
+			{
+				ResetMeshComponent(ParentComponent, InTargetMeshComponent);
+			}
+			else
+			{
+				CurrentlyUsedSourceMeshComponent.Reset();
+			}
+		}
+		else
+		{
+			CurrentlyUsedSourceMeshComponent.Reset();
+		}
+	}
+	else
+	{
+		CurrentlyUsedSourceMeshComponent.Reset();
+	}
+}
diff --git a/Engine/Plugins/Experimental/Animation/CurveExpression/Source/Runtime/Private/ExpressionEvaluator.cpp b/Engine/Plugins/Experimental/Animation/CurveExpression/Source/Runtime/Private/ExpressionEvaluator.cpp
new file mode 100644
--- /dev/null
+++ b/Engine/Plugins/Experimental/Animation/CurveExpression/Source/Runtime/Private/ExpressionEvaluator.cpp
@@ -0,0 +1,727 @@
+// Copyright Epic Games, Inc. All Rights Reserved.
+
+
+#include "ExpressionEvaluator.h"
+
+#include <limits>
+
+namespace CurveExpression::Evaluator
+{
+
+FEngine::FEngine(
+	TMap<FName, float>&& InConstants
+	) :
+	Constants(MoveTemp(InConstants))
+{
+}
+
+
+void FEngine::UpdateConstantValues(
+	const TMap<FName, float>& InConstants
+	)
+{
+	for (const TTuple<FName, float>& NewConstantItem: InConstants)
+	{
+		if (float* OldConstant = Constants.Find(NewConstantItem.Key))
+		{
+			*OldConstant = NewConstantItem.Value; 
+		}
+	}
+}
+
+// We can't use FChar::IsDigit, since it obeys locale, which would allow all sorts of non-ASCII
+// numeric digits from Unicode.
+static bool IsASCIIDigit(TCHAR InCh)
+{
+	return InCh >= '0' && InCh <= '9';
+}
+
+TVariant<FEngine::FToken, FParseError> FEngine::ParseTokenError(
+	FString&& InError,
+	const TCHAR* InCurrentPos,
+	FStringView InExpression)
+{
+	return TVariant<FToken, FParseError>(TInPlaceType<FParseError>(),
+		static_cast<int32>(InCurrentPos - InExpression.GetData()), MoveTemp(InError));
+}
+
+
+FString FEngine::TokenToString(
+	const FEngine::FToken& InToken
+	)
+{
+	const TCHAR* TypeStr = TEXT("<unknown>");
+	FString ValueStr;
+	
+	if (const EOperatorToken* TokenKind = InToken.TryGet<EOperatorToken>())
+	{
+		TypeStr = TEXT("Token");
+		switch(*TokenKind)
+		{
+		case EOperatorToken::Negate:		ValueStr = TEXT("<->"); break;
+		case EOperatorToken::Add:			ValueStr = TEXT("+"); break;
+		case EOperatorToken::Subtract:		ValueStr = TEXT("-"); break;
+		case EOperatorToken::Multiply:		ValueStr = TEXT("*"); break;
+		case EOperatorToken::Divide:		ValueStr = TEXT("/"); break;
+		case EOperatorToken::Modulo:		ValueStr = TEXT("%"); break;
+		case EOperatorToken::Power:			ValueStr = TEXT("**"); break;
+		case EOperatorToken::FloorDivide:	ValueStr = TEXT("//"); break;
+		case EOperatorToken::ParenOpen:		ValueStr = TEXT("("); break;
+		case EOperatorToken::ParenClose:	ValueStr = TEXT(")"); break;
+		}
+	}
+	else if (const FName* Identifier = InToken.TryGet<FName>())
+	{
+		TypeStr = TEXT("Identifier");
+		ValueStr = Identifier->ToString();
+	}
+	else if (const float* Value = InToken.TryGet<float>())
+	{
+		TypeStr = TEXT("Value");
+		ValueStr = FString::Printf(TEXT("%g"), *Value);
+	}
+
+	return FString::Printf(TEXT("%s[%s]"), TypeStr, *ValueStr);
+}
+
+
+TVariant<FEngine::FToken, FParseError> FEngine::ParseFloat(
+	FStringView& InOutParseRange,
+	FStringView InExpression
+	)
+{
+	// We only parse at most 11 significant digits, since 9 is the maximum decimal precision 
+	// a float will hold. We store 2 extra digits to make sure we round properly on the 
+	// last digit bit.
+	constexpr int32 MaxPrecision = 11;
+
+	// Count the number of digits in the mantissa and, optionally, locate the decimal point
+	// and the start of significant digits.
+	int32 DecimalPointPos = -1;
+	int32 MantissaLen = 0;
+	int32 SignificantDigitStart = -1;
+	for (; MantissaLen < InOutParseRange.Len(); MantissaLen++)
+	{
+		const TCHAR Ch = InOutParseRange[MantissaLen];
+		if (!IsASCIIDigit(Ch))
+		{
+			// If the current character is neither a digit or period, or if we've already
+			// found a period, stop counting.
+			if (Ch != '.' || DecimalPointPos >= 0)
+			{
+				break;
+			}
+			DecimalPointPos = MantissaLen;
+		}
+		else if (Ch != '0' && SignificantDigitStart < 0)
+		{
+			SignificantDigitStart = MantissaLen;
+		}
+	}
+
+	// Remember the position for the exponent parsing.
+	int32 LastPos = MantissaLen;
+	int32 SignificantDigitCount = SignificantDigitStart >= 0 ? MantissaLen - SignificantDigitStart : 0;
+
+	// If no decimal point was given, set the position of it after the mantissa. Otherwise
+	// remove one from the mantissa length, since that would be the decimal point.
+	if (DecimalPointPos < 0)
+	{
+		DecimalPointPos = MantissaLen;
+	}
+	else
+	{
+		MantissaLen--;		
+
+		// If the significant digit starts before the decimal point, remove one to account
+		// for the decimal point.
+		if (SignificantDigitStart < DecimalPointPos)
+		{
+			SignificantDigitCount--;
+		}
+	}
+	
+	if (MantissaLen == 0)
+	{
+		// We only parsed the period. That's not a valid float.
+		return ParseTokenError(TEXT("Invalid floating point value"), InOutParseRange.GetData(), InExpression);
+	}
+
+	// Adjustment for the exponent based on number of significant digits in the fraction part.
+	int32 FractionExponent;
+	if (SignificantDigitCount > MaxPrecision)
+	{
+		FractionExponent = DecimalPointPos - MaxPrecision;
+		SignificantDigitCount = MaxPrecision;
+	}
+	else
+	{
+		FractionExponent = DecimalPointPos - MantissaLen;
+	}
+
+	// Copy the integer and fraction digits, up to the maximum count.
+	uint64 IntegerValue = 0;
+	const TCHAR *Digits = InOutParseRange.GetData() + SignificantDigitStart;
+	for (int32 DigitIndex = 0; DigitIndex < SignificantDigitCount; DigitIndex++)
+	{
+		TCHAR Ch = *Digits++;
+		if (Ch == '.')
+		{
+			Ch = *Digits++;
+		}
+		IntegerValue = 10 * IntegerValue + (Ch - '0');
+	}
+	double Fraction = static_cast<double>(IntegerValue);
+
+	// Now that we have the fraction part, parse the exponent, if any.
+	int32 Exponent = 0;
+	bool IsExponentNegative = false;
+
+	// We need at least _two_ characters here ('e' + digit)
+	if ((LastPos + 1) < InOutParseRange.Len() &&
+		(InOutParseRange[LastPos] == 'e' || InOutParseRange[LastPos] == 'E'))
+	{
+		int32 ExponentPos = LastPos + 1;
+		if (InOutParseRange[ExponentPos] == '+')
+		{
+			ExponentPos++;
+		}
+		else if (InOutParseRange[ExponentPos] == '-')
+		{
+			IsExponentNegative = true;
+			ExponentPos++;
+		}
+
+		int32 DigitCount = 0;
+		while(ExponentPos < InOutParseRange.Len() && IsASCIIDigit(InOutParseRange[ExponentPos]))
+		{
+			// We don't care past 2 digits. 38 is the maximum exponent we care about for floats.
+			if (DigitCount < 2)
+			{
+				Exponent = Exponent * 10 + (InOutParseRange[ExponentPos] - '0');
+			}
+			ExponentPos++;
+			DigitCount++;
+		}
+
+		// Only if we parsed some digits do we consider the exponent parsed at all. Otherwise
+		// we leave the last parse at where the 'e' started.
+		if (DigitCount > 0)
+		{
+			LastPos = ExponentPos;
+		}
+	}
+
+	// Adjust the exponent to deal with the decimal point location.
+	if (IsExponentNegative)
+	{
+		Exponent = FractionExponent - Exponent;
+	}
+	else
+	{
+		Exponent = FractionExponent + Exponent;
+	}
+
+	// The maximum exponent for floats is 38. 
+	if (Exponent < 0) 
+	{
+		IsExponentNegative = true;
+		Exponent = -Exponent;
+	} 
+	else 
+	{
+		IsExponentNegative = false;
+	}
+
+	Exponent = FMath::Min(Exponent, std::numeric_limits<float>::max_exponent10);
+	double DoubleExponent = 1.0;
+
+	for(const double PowerOf2Exp: {1.0e1, 1.0e2, 1.0e4, 1.0e8, 1.0e16, 1.0e32})
+	{
+		if (Exponent == 0)
+		{
+			break;
+		}
+		if (Exponent & 1)
+		{
+			DoubleExponent *= PowerOf2Exp;
+		}
+		Exponent >>= 1;
+	}
+
+	if (IsExponentNegative)
+	{
+		Fraction /= DoubleExponent;
+	}
+	else
+	{
+		Fraction *= DoubleExponent;
+	}
+
+	InOutParseRange = InOutParseRange.Mid(LastPos);
+	
+	return TVariant<FToken, FParseError>(
+		TInPlaceType<FToken>(), FToken(TInPlaceType<float>(), Fraction));
+}
+
+
+TVariant<FEngine::FToken, FParseError> FEngine::ParseIdentifier(
+	FStringView& InOutParseRange,
+	FStringView InExpression
+	) const
+{
+	// FIXME: We need quote escaping.
+	bool bStopAtSingleQuote = false;
+	int32 ParseIndex = 0;
+	if (InOutParseRange[0] == '\'')
+	{
+		bStopAtSingleQuote = true;
+		ParseIndex++;
+		if (ParseIndex >= InOutParseRange.Len())
+		{
+			return ParseTokenError(TEXT("Missing end quote for constant"), InOutParseRange.GetData(), InExpression);
+		}
+	}
+	// Unquoted constants have to start with an alphabetic character.
+	else if (!FChar::IsAlpha(InOutParseRange[0]) && InOutParseRange[0] != '_')
+	{
+		return ParseTokenError(TEXT("Unexpected character"), InOutParseRange.GetData(), InExpression);
+	}
+
+	const int32 ConstantStart = ParseIndex;
+	int32 ConstantEnd;
+
+	ParseIndex++;
+
+	if (bStopAtSingleQuote)
+	{
+		while(ParseIndex < InOutParseRange.Len() && InOutParseRange[ParseIndex] != '\'')
+		{
+			ParseIndex++;
+		}
+
+		if (ParseIndex == InOutParseRange.Len())
+		{
+			return ParseTokenError(TEXT("Missing end quote for constant"), InOutParseRange.GetData(), InExpression);
+		}
+
+		ConstantEnd = ParseIndex;
+		ParseIndex++;
+	}
+	else
+	{
+		while(ParseIndex < InOutParseRange.Len() &&
+			  (FChar::IsAlnum(InOutParseRange[ParseIndex]) || InOutParseRange[ParseIndex] == '_'))
+		{
+			ParseIndex++;
+		}
+		
+		ConstantEnd = ParseIndex;
+	}
+
+	// Look up this identifier. We don't want to create a new FName, to avoid polluting the
+	// FNameEntry cache.
+	const FStringView Identifier(InOutParseRange.SubStr(ConstantStart, ConstantEnd - ConstantStart));
+	const FName Name(Identifier.Len(), Identifier.GetData(), FNAME_Find);
+	
+	InOutParseRange = InOutParseRange.Mid(ParseIndex); 
+
+	if (Constants.Contains(Name))
+	{
+		return TVariant<FToken, FParseError>(
+			TInPlaceType<FToken>(), FToken(TInPlaceType<FName>(), Name));
+	}
+	else
+	{
+		// If the constant value is not found, it gets a default value of 0.
+		// TBD: We may want to flag this as an error instead when doing verification.
+		return TVariant<FToken, FParseError>(
+			TInPlaceType<FToken>(), FToken(TInPlaceType<float>(), 0.0f));
+	}
+}
+
+
+TVariant<FEngine::FToken, FParseError> FEngine::ParseToken(
+	FStringView& InOutParseRange,
+	FStringView InExpression
+	) const
+{
+	auto TokenResult = [](EOperatorToken InTokenKind)
+	{
+		return TVariant<FToken, FParseError>(
+			TInPlaceType<FToken>(), FToken(TInPlaceType<EOperatorToken>(), InTokenKind));
+	};
+
+	switch(InOutParseRange[0])
+	{
+	case '+':
+		InOutParseRange = InOutParseRange.Mid(1);
+		return TokenResult(EOperatorToken::Add);
+		
+	case '-':
+		InOutParseRange = InOutParseRange.Mid(1);
+		return TokenResult(EOperatorToken::Subtract);
+		
+	case '*':
+		if (InOutParseRange.Len() >= 2 && InOutParseRange[1] == '*')
+		{
+			InOutParseRange = InOutParseRange.Mid(2);
+			return TokenResult(EOperatorToken::Power);
+		}
+		InOutParseRange = InOutParseRange.Mid(1);
+		return TokenResult(EOperatorToken::Multiply);
+		
+	case '/':
+		if (InOutParseRange.Len() >= 2 && InOutParseRange[1] == '/')
+		{
+			InOutParseRange = InOutParseRange.Mid(2);
+			return TokenResult(EOperatorToken::FloorDivide);
+		}
+		InOutParseRange = InOutParseRange.Mid(1);
+		return TokenResult(EOperatorToken::Divide);
+		
+	case '%':
+		InOutParseRange = InOutParseRange.Mid(1);
+		return TokenResult(EOperatorToken::Modulo);
+
+	case '(':
+		InOutParseRange = InOutParseRange.Mid(1);
+		return TokenResult(EOperatorToken::ParenOpen);
+		
+	case ')':
+		InOutParseRange = InOutParseRange.Mid(1);
+		return TokenResult(EOperatorToken::ParenClose);
+		
+	default:
+		if (IsASCIIDigit(InOutParseRange[0]) || InOutParseRange[0] == '.')
+		{
+			return ParseFloat(InOutParseRange, InExpression);
+		}
+		return ParseIdentifier(InOutParseRange, InExpression);
+	}
+}
+
+
+TVariant<FExpressionObject, FParseError> FEngine::Parse(
+	FStringView InExpression
+	) const
+{
+	enum class EAssociativity
+	{
+		None,
+		Left,
+		Right
+	};
+	struct FOperatorTokenInfo
+	{
+		EOperatorToken OperatorToken;
+		int32 Precedence;
+		EAssociativity Associativity;
+		TOptional<FExpressionObject::EOperator> Operator;
+	};
+	// Keep in same order as EOperatorToken. One day, if we're lucky, we'll get reflections.
+	auto GetOperatorTokenInfo = [](EOperatorToken InOpToken)
+	{
+		static const FOperatorTokenInfo OperatorTokenInfo[] = {
+			{ EOperatorToken::Negate,		5, EAssociativity::Right, FExpressionObject::EOperator::Negate},
+			{ EOperatorToken::Add,			1, EAssociativity::Left, FExpressionObject::EOperator::Add},
+			{ EOperatorToken::Subtract,		1, EAssociativity::Left, FExpressionObject::EOperator::Subtract},
+			{ EOperatorToken::Multiply,		2, EAssociativity::Left, FExpressionObject::EOperator::Multiply},
+			{ EOperatorToken::Divide,		2, EAssociativity::Left, FExpressionObject::EOperator::Divide},
+			{ EOperatorToken::Modulo,		3, EAssociativity::Left, FExpressionObject::EOperator::Modulo},
+			{ EOperatorToken::Power,		4, EAssociativity::Right, FExpressionObject::EOperator::Power},
+			{ EOperatorToken::FloorDivide,	3, EAssociativity::Left, FExpressionObject::EOperator::FloorDivide},
+			{ EOperatorToken::ParenOpen,	0, EAssociativity::None},
+			{ EOperatorToken::ParenClose,	0, EAssociativity::None},
+		};
+		checkSlow(OperatorTokenInfo[static_cast<int32>(InOpToken)].OperatorToken == Op);
+		return OperatorTokenInfo[static_cast<int32>(InOpToken)];
+	};
+	
+	static TVariant<FExpressionObject, FParseError> EmptyResult{TInPlaceType<FExpressionObject>()};
+
+	TStringView ParseRange = InExpression.TrimStartAndEnd();
+	if (ParseRange.IsEmpty())
+	{
+		return EmptyResult;
+	}
+
+	// Run Dijkstra's classic Shunting Yard algorithm to convert infix expressions to RPN.
+	// TODO: Add support for calling functions.
+	// TODO: Better error pinpointing by storing the token parse result with the token.
+	TArray<FExpressionObject::OpElement, TInlineAllocator<64>> Expression;
+	TArray<EOperatorToken, TInlineAllocator<32>> OperatorStack;
+
+	auto ParseError = [InExpression](FString&& InError, const TCHAR* InCurrentPos)
+	{
+		return TVariant<FExpressionObject, FParseError>(TInPlaceType<FParseError>(),
+			static_cast<int32>(InCurrentPos - InExpression.GetData()), MoveTemp(InError));
+	};
+	auto PushOperatorTokenToExpression = [GetOperatorTokenInfo, &Expression](EOperatorToken InOpToken)
+	{
+		Expression.Push(FExpressionObject::OpElement(TInPlaceType<FExpressionObject::EOperator>(), 
+			*GetOperatorTokenInfo(InOpToken).Operator));
+	};
+	// Used to check if we encounter a '-' to see if we need to turn it into a negate operator
+	// and also to see if we have two operators in a row (error).
+	TOptional<FToken> LastToken;
+
+	do
+	{
+		const TCHAR* TokenStart = ParseRange.GetData();
+		TVariant<FToken, FParseError> ParseResult = ParseToken(ParseRange, InExpression);
+		
+		if (FParseError* Error = ParseResult.TryGet<FParseError>())
+		{
+			return TVariant<FExpressionObject, FParseError>(TInPlaceType<FParseError>(), *Error);
+		}
+
+		const FStringView TokenParseRange(TokenStart, ParseRange.GetData() - TokenStart);
+		const FToken& Token = ParseResult.Get<FToken>();
+
+		// Values / identifiers go onto the expression stack immediately.
+		if (const FName* Identifier = Token.TryGet<FName>())
+		{
+			if (LastToken.IsSet() &&
+				(LastToken->IsType<FName>() || LastToken->IsType<float>()))
+			{
+				return ParseError(TEXT("Expected an operator"), TokenParseRange.GetData());
+			}
+			Expression.Push(FExpressionObject::OpElement(TInPlaceType<FName>(), *Identifier));
+		}
+		else if (const float* Value = Token.TryGet<float>())
+		{
+			if (LastToken.IsSet() &&
+				(LastToken->IsType<FName>() || LastToken->IsType<float>()))
+			{
+				return ParseError(TEXT("Expected an operator"), TokenParseRange.GetData());
+			}
+			Expression.Push(FExpressionObject::OpElement(TInPlaceType<float>(), *Value));
+		}
+		else
+		{
+			EOperatorToken Op = Token.Get<EOperatorToken>();
+			const FOperatorTokenInfo& TokenInfo = GetOperatorTokenInfo(Op);
+
+			// Was the last token also an operator token?
+			if (LastToken.IsSet())
+			{
+				if (const EOperatorToken* LastOp = LastToken->TryGet<EOperatorToken>())
+				{
+					const FOperatorTokenInfo& LastTokenInfo = GetOperatorTokenInfo(*LastOp);
+					if (TokenInfo.Operator.IsSet() && LastTokenInfo.Operator.IsSet())
+					{
+						return ParseError(TEXT("Expected an expression"), TokenParseRange.GetData());
+					}
+					if (*LastOp == EOperatorToken::ParenClose && Op == EOperatorToken::ParenOpen)
+					{
+						return ParseError(TEXT("Expected an operator"), TokenParseRange.GetData());
+					}
+					if (*LastOp == EOperatorToken::ParenOpen && Op == EOperatorToken::ParenClose)
+					{
+						return ParseError(TEXT("Empty Parentheses"), TokenParseRange.GetData());
+					}
+				}
+			}
+
+			// Special-case for unary prefix.
+			if (Op == EOperatorToken::Subtract &&
+				(!LastToken.IsSet() || (LastToken->IsType<EOperatorToken>() && LastToken->Get<EOperatorToken>() != EOperatorToken::ParenClose)))
+			{
+				Op = EOperatorToken::Negate;
+				OperatorStack.Push(Op);
+			}
+			else if (Op == EOperatorToken::ParenOpen)
+			{
+				OperatorStack.Push(Op);
+			}
+			else if (Op == EOperatorToken::ParenClose)
+			{
+				while(OperatorStack.IsEmpty() || OperatorStack.Last() != EOperatorToken::ParenOpen)
+				{
+					if (OperatorStack.IsEmpty())
+					{
+						return ParseError(TEXT("Mismatched parentheses"), TokenParseRange.GetData()); 
+					}
+					PushOperatorTokenToExpression(OperatorStack.Pop(false));
+				}
+				// Remove the now unneeded open parentheses.
+				OperatorStack.Pop(false);
+			}
+			else if (GetOperatorTokenInfo(Op).Associativity == EAssociativity::Left)
+			{
+				// Peel off all operators that have the same precedence or higher.
+				while (!OperatorStack.IsEmpty() && GetOperatorTokenInfo(OperatorStack.Last()).Precedence >= TokenInfo.Precedence)
+				{
+					PushOperatorTokenToExpression(OperatorStack.Pop(false));
+				}
+				OperatorStack.Push(Op);
+			}
+			else if (GetOperatorTokenInfo(Op).Associativity == EAssociativity::Right)
+			{
+				// Peel off all operators that have higher precedence.
+				while (!OperatorStack.IsEmpty() && GetOperatorTokenInfo(OperatorStack.Last()).Precedence > TokenInfo.Precedence)
+				{
+					PushOperatorTokenToExpression(OperatorStack.Pop(false));
+				}
+				OperatorStack.Push(Op);
+			}
+		}
+
+		LastToken = Token;
+		ParseRange = ParseRange.TrimStart();
+	} while(!ParseRange.IsEmpty());
+
+	while (!OperatorStack.IsEmpty())
+	{
+		if (!GetOperatorTokenInfo(OperatorStack.Last()).Operator.IsSet())
+		{
+			return ParseError(TEXT("Mismatched parentheses"), InExpression.GetData()); 
+		}
+		PushOperatorTokenToExpression(OperatorStack.Pop(false));
+	}
+
+	FExpressionObject S;
+	S.Expression = Expression;
+	return TVariant<FExpressionObject, FParseError>(TInPlaceType<FExpressionObject>(), S);
+}
+
+
+float FEngine::Execute(
+	const FExpressionObject& InExpressionObject
+	) const
+{
+	// An empty expression object is valid.
+	if (InExpressionObject.Expression.IsEmpty())
+	{
+		return 0.0f;
+	}
+
+	// We assume the codons in the expression object are correct and will always result
+	// in a single value remaining on the stack.
+	TArray<float, TInlineAllocator<32>> ValueStack;
+	
+	for (const FExpressionObject::OpElement& Token: InExpressionObject.Expression)
+	{
+		if (const FExpressionObject::EOperator* Operator = Token.TryGet<FExpressionObject::EOperator>())
+		{
+			switch(*Operator)
+			{
+			case FExpressionObject::EOperator::Negate:
+				{
+					ValueStack.Last() = -ValueStack.Last();
+				}
+				break;	
+				
+			case FExpressionObject::EOperator::Add:
+				{
+					const float V = ValueStack.Pop(false);
+					ValueStack.Last() += V;
+				}
+				break;	
+			case FExpressionObject::EOperator::Subtract:
+				{
+					const float V = ValueStack.Pop(false);
+					ValueStack.Last() -= V;
+				}
+				break;	
+			case FExpressionObject::EOperator::Multiply:
+				{
+					const float V = ValueStack.Pop(false);
+					ValueStack.Last() *= V;
+				}
+				break;
+			case FExpressionObject::EOperator::Divide:
+				{
+					const float V = ValueStack.Pop(false);
+					if (FMath::IsNearlyZero(V))
+					{
+						ValueStack.Last() = 0.0f;
+					}
+					else
+					{
+						ValueStack.Last() /= V;
+					}
+				}
+				break;	
+			case FExpressionObject::EOperator::Modulo:
+				{
+					const float V = ValueStack.Pop(false);
+					if (FMath::IsNearlyZero(V))
+					{
+						ValueStack.Last() = 0.0f;
+					}
+					else
+					{
+						ValueStack.Last() = FMath::Fmod(ValueStack.Last(), V);
+					}
+				}
+				break;	
+			case FExpressionObject::EOperator::Power:
+				{
+					const float V = ValueStack.Pop(false);
+					ValueStack.Last() = FMath::Pow(ValueStack.Last(), V);
+					if (!FMath::IsFinite(ValueStack.Last()))
+					{
+						ValueStack.Last() = 0.0f;
+					}
+				}
+				break;	
+			case FExpressionObject::EOperator::FloorDivide:
+				{
+					const float V = ValueStack.Pop(false);
+					if (FMath::IsNearlyZero(V))
+					{
+						ValueStack.Last() = 0.0f;
+					}
+					else
+					{
+						ValueStack.Last() /= V;
+					}
+					ValueStack.Last() = FMath::Floor(ValueStack.Last());
+				}
+				break;  
+			}
+		}
+		else if (const FName* Constant = Token.TryGet<FName>())
+		{
+			const float *Value = Constants.Find(*Constant);
+			ValueStack.Push(Value ? *Value : 0.0f);
+		}
+		else if (const float* Value = Token.TryGet<float>())
+		{
+			ValueStack.Push(*Value);
+		}
+	}
+
+	return ValueStack.Last();
+}
+
+
+TOptional<float> FEngine::Evaluate(
+	FStringView InExpression
+	) const
+{
+	TVariant<FExpressionObject, FParseError> Result = Parse(InExpression);
+	if (const FExpressionObject* ExpressionObject = Result.TryGet<FExpressionObject>())
+	{
+		return Execute(*ExpressionObject);
+	}
+
+	return {};
+}
+
+
+TOptional<FParseError> FEngine::Verify(
+	FStringView InExpression
+	) const
+{
+	TVariant<FExpressionObject, FParseError> Result = Parse(InExpression);
+	if (const FParseError* Error = Result.TryGet<FParseError>())
+	{
+		return *Error;
+	}
+	return {};
+}
+
+}
diff --git a/Engine/Plugins/Experimental/Animation/CurveExpression/Source/Runtime/Public/AnimNode_RemapCurvesFromMesh.h b/Engine/Plugins/Experimental/Animation/CurveExpression/Source/Runtime/Public/AnimNode_RemapCurvesFromMesh.h
new file mode 100644
--- /dev/null
+++ b/Engine/Plugins/Experimental/Animation/CurveExpression/Source/Runtime/Public/AnimNode_RemapCurvesFromMesh.h
@@ -0,0 +1,58 @@
+// Copyright Epic Games, Inc. All Rights Reserved.
+
+#pragma once
+
+#include "ExpressionEvaluator.h"
+#include "Animation/AnimNodeBase.h"
+
+#include "AnimNode_RemapCurvesFromMesh.generated.h"
+
+USTRUCT(BlueprintInternalUseOnly)
+struct CURVEEXPRESSION_API FAnimNode_RemapCurvesFromMesh :
+	public FAnimNode_Base
+{
+	GENERATED_BODY()
+
+	UPROPERTY(EditAnywhere, EditFixedSize, BlueprintReadWrite, Category = Links)
+	FPoseLink SourcePose;
+	
+	/** This is used by default if it's valid */
+	UPROPERTY(BlueprintReadWrite, transient, Category=Copy, meta=(PinShownByDefault))
+	TWeakObjectPtr<USkeletalMeshComponent> SourceMeshComponent;
+
+	/** If SourceMeshComponent is not valid, and if this is true, it will look for attached parent as a source */
+	UPROPERTY(EditAnywhere, BlueprintReadWrite, Category=Copy, meta=(NeverAsPin))
+	bool bUseAttachedParent = false;
+	
+	UPROPERTY(EditAnywhere, BlueprintReadWrite, editfixedsize, Category = Expressions, meta = (PinHiddenByDefault))
+	TMap<FName, FString> CurveExpressions;
+
+	/** The expressions given are immutable and will not change during runtime. Improves performance. */ 
+	UPROPERTY(EditAnywhere, BlueprintReadWrite, Category=Expressions, meta=(NeverAsPin))
+	bool bExpressionsImmutable = true;
+	
+	// FAnimNode_Base interface
+	virtual void Initialize_AnyThread(const FAnimationInitializeContext& Context) override;
+	virtual void CacheBones_AnyThread(const FAnimationCacheBonesContext& Context) override;
+	virtual void Update_AnyThread(const FAnimationUpdateContext& Context) override;
+	virtual void Evaluate_AnyThread(FPoseContext& Output) override;
+	virtual void GatherDebugData(FNodeDebugData& DebugData) override;
+	virtual bool HasPreUpdate() const override { return true; }
+	virtual void PreUpdate(const UAnimInstance* InAnimInstance) override;
+	
+private:
+	void ReinitializeMeshComponent(USkeletalMeshComponent* InNewSkeletalMeshComponent, USkeletalMeshComponent* InTargetMeshComponent);
+	void RefreshMeshComponent(USkeletalMeshComponent* InTargetMeshComponent);
+
+	// this is source mesh references, so that we could compare and see if it has changed
+	TWeakObjectPtr<USkeletalMeshComponent>	CurrentlyUsedSourceMeshComponent;
+	TWeakObjectPtr<USkeletalMesh>			CurrentlyUsedSourceMesh;
+
+	// target mesh 
+	TWeakObjectPtr<USkeletalMesh>			CurrentlyUsedTargetMesh;
+	
+	TMap<FName, SmartName::UID_Type> CurveNameToUIDMap;
+
+	TOptional<CurveExpression::Evaluator::FEngine> ExpressionEngine;
+	TMap<FName, CurveExpression::Evaluator::FExpressionObject> CachedExpressions;
+};
diff --git a/Engine/Plugins/Experimental/Animation/CurveExpression/Source/Runtime/Public/ExpressionEvaluator.h b/Engine/Plugins/Experimental/Animation/CurveExpression/Source/Runtime/Public/ExpressionEvaluator.h
new file mode 100644
--- /dev/null
+++ b/Engine/Plugins/Experimental/Animation/CurveExpression/Source/Runtime/Public/ExpressionEvaluator.h
@@ -0,0 +1,169 @@
+// Copyright Epic Games, Inc. All Rights Reserved.
+
+#pragma once
+
+#include "Containers/Map.h"
+#include "Containers/UnrealString.h"
+#include "Misc/Optional.h"
+#include "Misc/TVariant.h"
+#include "UObject/NameTypes.h"
+
+namespace CurveExpression::Evaluator
+{
+
+struct FParseError
+{
+	FParseError() = default;
+	FParseError(const int32 InColumn, FString&& InMessage) : Column(InColumn), Message(MoveTemp(InMessage)) {}
+	
+	/** The column at which the error occurs. Zero-based. */
+	const int32 Column = 0;
+
+	/** The human-readable parsing error */
+	const FString Message;
+};
+
+struct FExpressionObject
+{
+private:
+	enum class EOperator
+	{
+		Negate,				// Negation operator.
+		Add,				// Add last two values on stack and add the result to the stack.
+		Subtract,			// Same but subtract 
+		Multiply,			// Same but multiply
+		Divide,				// Same but divide (div-by-zero returns zero)
+		Modulo,				// Same but modulo (mod-by-zero returns zero)
+		Power,				// Raise to power
+		FloorDivide,    	// Divide and round the result down
+	};
+	
+	friend class FEngine;
+	using OpElement = TVariant<EOperator, /* Constant Ref */ FName, /* Value */ float>;
+	TArray<OpElement, TInlineAllocator<8>> Expression;
+};
+
+/** Simple, fast expression evaluation engine. Can evaluate simple arithmetic expressions,
+ *  including user-settable constants. 
+ *		1 + 1
+ *		foo + (3.0 * bar)
+ *
+ *  Allowable operands are taken from Python and are currently as follow:
+ *      -x				Unary negation
+ *      x + y			Addition
+ *      x - y			Subtraction
+ *      x * y			Multiplication
+ *      x / y			Division (div-by-zero returns zero)
+ *      x ** y			Raise x to y's power
+ *      x // y			Division and rounding down.
+ *
+ *  Terms can also be grouped by parentheses '(' and ')'.
+ *  
+ *	Constants that begin with a number, include spaces or any of the above operands/parentheses 
+ *	are allowed, but must be then enclosed in single quotes when referenced.
+ *	E.g:
+ *	   'foo(1)' * 2.0
+ *	rather than:
+ *	   foo (1) * 2.0
+ */
+class FEngine
+{
+public:
+	FEngine() = default;
+	
+	/** Initialize the engine with constants */
+	explicit FEngine(
+		TMap<FName, float>&& InConstants
+		);
+
+	/** Update existing constant values. No new constants are added. */
+	void UpdateConstantValues(
+		const TMap<FName, float>& InConstants
+		);
+
+	/** Parse an expression and either receive an error or a expression execution object that
+	 *  matches this evaluator. */
+	TVariant<FExpressionObject, FParseError> Parse(
+		FStringView InExpression
+		) const;
+
+	/** Executes the given expression object */
+	float Execute(
+		const FExpressionObject& InExpressionObject
+		) const;
+	
+	/** Evaluate an expression directly and returns the result. If the expression is incorrectly
+	 *  formed, then the optional return value holds nothing.
+	 */ 
+	TOptional<float> Evaluate(
+		FStringView InExpression
+		) const;
+
+	/** Verify an expression. Does not check for variable validity, only that the expression
+	 *  can be evaluated */
+	TOptional<FParseError> Verify(
+		FStringView InExpression
+		) const;
+	
+private:
+	enum class EOperatorToken
+	{
+		Negate,				// Placeholder for when constructing the op stack. Gets converted to (0 - V)
+		Add,				// '+'
+		Subtract,			// '-'
+		Multiply,			// '*'
+		Divide,				// '/'
+		Modulo,				// '%'
+		Power,				// '**'
+		FloorDivide,    	// '//'
+		ParenOpen,			// '('
+		ParenClose,			// ')'
+	};
+	using FToken = TVariant<EOperatorToken, /* Identifier */ FName, /* Value */ float>;
+	
+	/** Parse a floating point value and return the result (or error). Same semantics
+	 *  apply as with ParseToken.
+	 */
+	static TVariant<FToken, FParseError> ParseFloat(
+		FStringView& InOutParseRange,
+		FStringView InExpression
+		);
+
+	/** Parse a constant reference and return the result (or error). Same semantics
+	 *  apply as with ParseToken.
+	 */
+	TVariant<FToken, FParseError> ParseIdentifier(
+		FStringView& InOutParseRange,
+		FStringView InExpression
+		) const;
+	
+	/** Parse a single token from the beginning of the range, returning either the token
+	 *  or an error, if one was encountered.
+	 *  If the token was successfully parsed, then the given string view range gets updated
+	 *  with the new starting point at the end of the token just parsed
+	 */
+	TVariant<FToken, FParseError> ParseToken(
+		FStringView& InOutParseRange,
+		FStringView InExpression
+		) const;
+
+	/** Helper function to construct a token parse error result */
+	static TVariant<FToken, FParseError> ParseTokenError(
+		FString&& InError,
+		const TCHAR* InCurrentPos,
+		FStringView InExpression
+		);
+	
+	/** Helper function to get a text version of a token. */
+	static FString TokenToString(
+		const FToken& InToken
+		);
+	
+	/** The list of constants, stored using the FName key since the expression object refers
+	 *  to it directly */
+	// FUTURE: Add support for variables by allowing TFunction{Ref} objects along with direct
+	// constants.
+	TMap<FName, float> Constants;
+};
+
+}
