Fix build break in RemapCurvesFromMesh

#rnx

diff --git a/Engine/Plugins/Experimental/Animation/CurveExpression/Source/Runtime/Private/AnimNode_RemapCurvesFromMesh.cpp b/Engine/Plugins/Experimental/Animation/CurveExpression/Source/Runtime/Private/AnimNode_RemapCurvesFromMesh.cpp
--- a/Engine/Plugins/Experimental/Animation/CurveExpression/Source/Runtime/Private/AnimNode_RemapCurvesFromMesh.cpp
+++ b/Engine/Plugins/Experimental/Animation/CurveExpression/Source/Runtime/Private/AnimNode_RemapCurvesFromMesh.cpp
@@ -188,11 +188,11 @@ void FAnimNode_RemapCurvesFromMesh::ReinitializeMeshComponent(
 
 				TMap<FName, float> SourceCurves;
 				
-				SourceContainer->Iterate([&SourceCurves](const FSmartNameMapping* Mapping, const SmartName::UID_Type ID)
+				SourceContainer->Iterate([&SourceCurves](const FSmartNameMappingIterator& Iterator)
 				{
-					if (FName CurveName; Mapping->GetName(ID, CurveName))
+					if (FName CurveName; Iterator.GetName(CurveName))
 					{
-						const FCurveMetaData* MetaData = Mapping->GetCurveMetaData(CurveName);
+						const FCurveMetaData* MetaData = Iterator.GetCurveMetaData();
 						FPlatformMisc::LowLevelOutputDebugStringf(TEXT("CURVE[%s] - Morph: %d - Mat: %d\n"), *CurveName.ToString(), MetaData->Type.bMaterial, MetaData->Type.bMorphtarget);
 						SourceCurves.Add(CurveName, 0.0f);
 					}
@@ -211,11 +211,11 @@ void FAnimNode_RemapCurvesFromMesh::ReinitializeMeshComponent(
 
 				// CurveExpressions isn't filled in here, since we haven't evaluated BP downstream.
 				// So instead, we just take a full copy of the smart mapping.
-				TargetContainer->Iterate([this](const FSmartNameMapping* Mapping, const SmartName::UID_Type ID)
+				TargetContainer->Iterate([this](const FSmartNameMappingIterator& Iterator)
 				{
-					if (FName CurveName; Mapping->GetName(ID, CurveName))
+					if (FName CurveName; Iterator.GetName(CurveName))
 					{
-						CurveNameToUIDMap.Add(CurveName, ID);
+						CurveNameToUIDMap.Add(CurveName, Iterator.GetIndex());
 					}
 				});
 			}
