Add missing module init files for CurveExpression plugin.

#rnx

diff --git a/Engine/Plugins/Experimental/Animation/CurveExpression/Source/Editor/Private/CurveExpressionEditorModule.cpp b/Engine/Plugins/Experimental/Animation/CurveExpression/Source/Editor/Private/CurveExpressionEditorModule.cpp
new file mode 100644
--- /dev/null
+++ b/Engine/Plugins/Experimental/Animation/CurveExpression/Source/Editor/Private/CurveExpressionEditorModule.cpp
@@ -0,0 +1,17 @@
+// Copyright Epic Games, Inc. All Rights Reserved.
+
+#include "Modules/ModuleManager.h"
+
+namespace UE::CurveExpressionEditor
+{
+
+class FModule : public IModuleInterface
+{
+public:
+	virtual void StartupModule() override {}
+	virtual void ShutdownModule() override {}
+};
+
+}
+
+IMPLEMENT_MODULE(UE::CurveExpressionEditor::FModule, CurveExpressionEditor)
diff --git a/Engine/Plugins/Experimental/Animation/CurveExpression/Source/Runtime/Private/CurveExpressionModule.cpp b/Engine/Plugins/Experimental/Animation/CurveExpression/Source/Runtime/Private/CurveExpressionModule.cpp
new file mode 100644
--- /dev/null
+++ b/Engine/Plugins/Experimental/Animation/CurveExpression/Source/Runtime/Private/CurveExpressionModule.cpp
@@ -0,0 +1,17 @@
+// Copyright Epic Games, Inc. All Rights Reserved.
+
+#include "Modules/ModuleManager.h"
+
+namespace UE::CurveExpression
+{
+
+class FModule : public IModuleInterface
+{
+public:
+	virtual void StartupModule() override {}
+	virtual void ShutdownModule() override {}
+};
+
+}
+
+IMPLEMENT_MODULE(UE::CurveExpression::FModule, CurveExpression)
