Fix Clang "explicitly defaulted default constructor is implicitly deleted"
#rnx

diff --git a/Engine/Plugins/Experimental/Animation/CurveExpression/Source/Runtime/Public/ExpressionEvaluator.h b/Engine/Plugins/Experimental/Animation/CurveExpression/Source/Runtime/Public/ExpressionEvaluator.h
--- a/Engine/Plugins/Experimental/Animation/CurveExpression/Source/Runtime/Public/ExpressionEvaluator.h
+++ b/Engine/Plugins/Experimental/Animation/CurveExpression/Source/Runtime/Public/ExpressionEvaluator.h
@@ -15,12 +15,12 @@ struct FParseError
 {
 	FParseError() = default;
 	FParseError(const int32 InColumn, FString&& InMessage) : Column(InColumn), Message(MoveTemp(InMessage)) {}
-	
+
 	/** The column at which the error occurs. Zero-based. */
 	const int32 Column = 0;
 
 	/** The human-readable parsing error */
-	const FString Message;
+	const FString Message = FString();
 };
 
 struct FExpressionObject
