Fix debug build
#rnx

diff --git a/Engine/Plugins/Experimental/Animation/CurveExpression/Source/Runtime/Private/ExpressionEvaluator.cpp b/Engine/Plugins/Experimental/Animation/CurveExpression/Source/Runtime/Private/ExpressionEvaluator.cpp
--- a/Engine/Plugins/Experimental/Animation/CurveExpression/Source/Runtime/Private/ExpressionEvaluator.cpp
+++ b/Engine/Plugins/Experimental/Animation/CurveExpression/Source/Runtime/Private/ExpressionEvaluator.cpp
@@ -434,7 +434,7 @@ TVariant<FExpressionObject, FParseError> FEngine::Parse(
 			{ EOperatorToken::ParenOpen,	0, EAssociativity::None},
 			{ EOperatorToken::ParenClose,	0, EAssociativity::None},
 		};
-		checkSlow(OperatorTokenInfo[static_cast<int32>(InOpToken)].OperatorToken == Op);
+		checkSlow(OperatorTokenInfo[static_cast<int32>(InOpToken)].OperatorToken == InOpToken);
 		return OperatorTokenInfo[static_cast<int32>(InOpToken)];
 	};
 	
