Added a convenience node to take a set of assignment expressions and break them into a key/value set of target curve/source curve expressions.

diff --git a/Engine/Plugins/Experimental/Animation/CurveExpression/Source/Editor/CurveExpressionEditor.Build.cs b/Engine/Plugins/Experimental/Animation/CurveExpression/Source/Editor/CurveExpressionEditor.Build.cs
--- a/Engine/Plugins/Experimental/Animation/CurveExpression/Source/Editor/CurveExpressionEditor.Build.cs
+++ b/Engine/Plugins/Experimental/Animation/CurveExpression/Source/Editor/CurveExpressionEditor.Build.cs
@@ -6,31 +6,18 @@ public class CurveExpressionEditor : ModuleRules
 {
 	public CurveExpressionEditor(ReadOnlyTargetRules Target) : base(Target)
 	{
-		PublicDependencyModuleNames.AddRange(
+		PrivateDependencyModuleNames.AddRange(
 			new string[]
 			{
 				"AnimGraph",
-				"AnimationCore",
+				"BlueprintGraph",
 				"Core",
 				"CoreUObject",
 				"CurveExpression",
 				"Engine",
-				
-				// UI 
-				"PropertyEditor",
+				"KismetCompiler",
 				"Slate",
 				"SlateCore",
-				"RewindDebuggerInterface",
-				"EditorStyle",
-				"UnrealEd",
-				"InputCore"
-			}
-		);
-
-		PrivateDependencyModuleNames.AddRange(
-			new string[]
-			{
-				"BlueprintGraph"
 			}
 		);
 	}
diff --git a/Engine/Plugins/Experimental/Animation/CurveExpression/Source/Editor/Private/CurveExpressionDetailsCustomization.cpp b/Engine/Plugins/Experimental/Animation/CurveExpression/Source/Editor/Private/CurveExpressionDetailsCustomization.cpp
new file mode 100644
--- /dev/null
+++ b/Engine/Plugins/Experimental/Animation/CurveExpression/Source/Editor/Private/CurveExpressionDetailsCustomization.cpp
@@ -0,0 +1,98 @@
+// Copyright Epic Games, Inc. All Rights Reserved.
+
+#include "CurveExpressionDetailsCustomization.h"
+
+#include "CurveExpressionEditorStyle.h"
+#include "K2Node_MakeCurveExpressionMap.h"
+#include "DetailWidgetRow.h"
+#include "Widgets/Layout/SGridPanel.h"
+#include "Widgets/Layout/SScrollBox.h"
+#include "Widgets/Text/SMultiLineEditableText.h"
+
+
+TSharedRef<IPropertyTypeCustomization> FCurveExpressionListCustomization::MakeInstance()
+{
+	return MakeShared<FCurveExpressionListCustomization>();
+}
+
+
+void FCurveExpressionListCustomization::CustomizeHeader(
+	TSharedRef<IPropertyHandle> InPropertyHandle,
+	FDetailWidgetRow& InHeaderRow,
+	IPropertyTypeCustomizationUtils& InCustomizationUtils
+	)
+{
+	AssignmentExpressionsProperty = InPropertyHandle->GetChildHandle(GET_MEMBER_NAME_CHECKED(FCurveExpressionList, AssignmentExpressions));
+
+	HorizontalScrollbar =
+		SNew(SScrollBar)
+			.AlwaysShowScrollbar(true)
+			.Orientation(Orient_Horizontal);
+
+	VerticalScrollbar =
+		SNew(SScrollBar)
+			.AlwaysShowScrollbar(true)
+			.Orientation(Orient_Vertical);
+
+	const FTextBlockStyle &TextStyle = FCurveExpressionEditorStyle::Get().GetWidgetStyle<FTextBlockStyle>("TextEditor.NormalText");
+	const FSlateFontInfo &Font = TextStyle.Font;
+	
+	InHeaderRow
+	.WholeRowContent()
+	[
+		SNew(SVerticalBox)
+		+ SVerticalBox::Slot()
+		.VAlign(VAlign_Top)
+		.AutoHeight()
+		[
+			InPropertyHandle->CreatePropertyNameWidget()
+		]
+		+ SVerticalBox::Slot()
+		.MaxHeight(400.0f)
+		[
+			SNew(SBorder)
+			.BorderImage(FCurveExpressionEditorStyle::Get().GetBrush("TextEditor.Border"))
+			.BorderBackgroundColor(FLinearColor::Black)
+			[
+				SNew(SGridPanel)
+				.FillColumn(0, 1.0f)
+				.FillRow(0, 1.0f)
+				+SGridPanel::Slot(0, 0)
+				[
+					SNew(SScrollBox)
+					.Orientation(Orient_Vertical)
+					.ExternalScrollbar(VerticalScrollbar)
+					+ SScrollBox::Slot()
+					[
+						SAssignNew(TextEditor, SMultiLineEditableText)
+						.Font(Font)
+						.TextStyle(&TextStyle)
+						.Text_Lambda([this]()
+						{
+							FString Text;
+							AssignmentExpressionsProperty->GetValue(Text);
+							return FText::FromString(Text);
+						})
+						.OnTextChanged_Lambda([this](const FText& InText)
+						{
+							AssignmentExpressionsProperty->SetValue(InText.ToString());
+						})
+						// By default, the Tab key gets routed to "next widget". We want to disable that behaviour.
+						.OnIsTypedCharValid_Lambda([](const TCHAR InChar) { return true; })
+						.AutoWrapText(false)
+						.HScrollBar(HorizontalScrollbar)
+						.VScrollBar(VerticalScrollbar)
+					]
+				]
+				+SGridPanel::Slot(1, 0)
+				[
+					VerticalScrollbar.ToSharedRef()
+				]
+				+SGridPanel::Slot(0, 1)
+				[
+					HorizontalScrollbar.ToSharedRef()
+				]
+			]
+		]
+	];
+}
diff --git a/Engine/Plugins/Experimental/Animation/CurveExpression/Source/Editor/Private/CurveExpressionDetailsCustomization.h b/Engine/Plugins/Experimental/Animation/CurveExpression/Source/Editor/Private/CurveExpressionDetailsCustomization.h
new file mode 100644
--- /dev/null
+++ b/Engine/Plugins/Experimental/Animation/CurveExpression/Source/Editor/Private/CurveExpressionDetailsCustomization.h
@@ -0,0 +1,36 @@
+// Copyright Epic Games, Inc. All Rights Reserved.
+
+#pragma once
+
+#include "IPropertyTypeCustomization.h"
+
+
+class SMultiLineEditableText;
+class SScrollBar;
+
+
+class FCurveExpressionListCustomization : 
+	public IPropertyTypeCustomization
+{
+public:
+	static TSharedRef<IPropertyTypeCustomization> MakeInstance();
+
+	// IPropertyTypeCustomization overrides
+	virtual void CustomizeHeader(
+		TSharedRef<IPropertyHandle> InPropertyHandle,
+		FDetailWidgetRow& InHeaderRow,
+		IPropertyTypeCustomizationUtils& InCustomizationUtils) override;
+
+	virtual void CustomizeChildren(
+		TSharedRef<IPropertyHandle> InPropertyHandle,
+		IDetailChildrenBuilder& InChildBuilder,
+		IPropertyTypeCustomizationUtils& InCustomizationUtils) override {}
+
+private:
+	TSharedPtr<IPropertyHandle> AssignmentExpressionsProperty;
+
+	TSharedPtr<SScrollBar> HorizontalScrollbar;
+	TSharedPtr<SScrollBar> VerticalScrollbar;
+
+	TSharedPtr<SMultiLineEditableText> TextEditor;	
+};
diff --git a/Engine/Plugins/Experimental/Animation/CurveExpression/Source/Editor/Private/CurveExpressionEditorModule.cpp b/Engine/Plugins/Experimental/Animation/CurveExpression/Source/Editor/Private/CurveExpressionEditorModule.cpp
--- a/Engine/Plugins/Experimental/Animation/CurveExpression/Source/Editor/Private/CurveExpressionEditorModule.cpp
+++ b/Engine/Plugins/Experimental/Animation/CurveExpression/Source/Editor/Private/CurveExpressionEditorModule.cpp
@@ -1,6 +1,10 @@
 // Copyright Epic Games, Inc. All Rights Reserved.
 
+#include "CurveExpressionDetailsCustomization.h"
+#include "CurveExpressionEditorStyle.h"
+#include "K2Node_MakeCurveExpressionMap.h"
 #include "Modules/ModuleManager.h"
+#include "PropertyEditorModule.h"
 
 namespace UE::CurveExpressionEditor
 {
@@ -8,8 +12,40 @@ namespace UE::CurveExpressionEditor
 class FModule : public IModuleInterface
 {
 public:
-	virtual void StartupModule() override {}
-	virtual void ShutdownModule() override {}
+	virtual void StartupModule() override
+	{
+		FPropertyEditorModule& PropertyModule = FModuleManager::LoadModuleChecked<FPropertyEditorModule>("PropertyEditor");
+		auto RegisterPropertyCustomization = [&](FName InStructName, auto InCustomizationFactory)
+		{
+			PropertyModule.RegisterCustomPropertyTypeLayout(
+				InStructName, 
+				FOnGetPropertyTypeCustomizationInstance::CreateStatic(InCustomizationFactory)
+				);
+			CustomizedProperties.Add(InStructName);
+		};
+		
+		FCurveExpressionEditorStyle::Register();
+		
+		RegisterPropertyCustomization(
+			FCurveExpressionList::StaticStruct()->GetFName(),
+			&FCurveExpressionListCustomization::MakeInstance
+			);
+	}
+	virtual void ShutdownModule() override
+	{
+		if (FPropertyEditorModule* PropertyModule = FModuleManager::GetModulePtr<FPropertyEditorModule>("PropertyEditor"))
+		{
+			for (const FName& PropertyName: CustomizedProperties)
+			{
+				PropertyModule->UnregisterCustomPropertyTypeLayout(PropertyName);
+			}
+		}
+		
+		FCurveExpressionEditorStyle::Unregister();
+	}
+
+private:
+	TArray<FName> CustomizedProperties;
 };
 
 }
diff --git a/Engine/Plugins/Experimental/Animation/CurveExpression/Source/Editor/Private/CurveExpressionEditorStyle.cpp b/Engine/Plugins/Experimental/Animation/CurveExpression/Source/Editor/Private/CurveExpressionEditorStyle.cpp
new file mode 100644
--- /dev/null
+++ b/Engine/Plugins/Experimental/Animation/CurveExpression/Source/Editor/Private/CurveExpressionEditorStyle.cpp
@@ -0,0 +1,67 @@
+// Copyright Epic Games, Inc. All Rights Reserved.
+
+#include "CurveExpressionEditorStyle.h"
+
+#include "Styling/CoreStyle.h"
+#include "Styling/SlateStyleMacros.h"
+#include "Styling/SlateStyleRegistry.h"
+#include "Styling/SlateTypes.h"
+
+
+FCurveExpressionEditorStyle& FCurveExpressionEditorStyle::Get()
+{
+	static FCurveExpressionEditorStyle Instance;
+	return Instance;
+}
+
+
+void FCurveExpressionEditorStyle::Register()
+{
+	FSlateStyleRegistry::RegisterSlateStyle(Get());
+}
+
+
+void FCurveExpressionEditorStyle::Unregister()
+{
+	FSlateStyleRegistry::UnRegisterSlateStyle(Get());
+}
+
+
+FCurveExpressionEditorStyle::FCurveExpressionEditorStyle() :
+	FSlateStyleSet("CurveExpressionEditorStyle")
+{
+	static const FSlateColor DefaultForeground(FLinearColor(0.72f, 0.72f, 0.72f, 1.f));
+
+	FSlateStyleSet::SetContentRoot(FPaths::EnginePluginsDir() / TEXT("Experimental/Animation/CurveExpression/Resources"));
+
+	// Text editor styles
+	{
+		const FSlateFontInfo Consolas10  = FCoreStyle::GetDefaultFontStyle("Mono", 9);
+
+		const FTextBlockStyle NormalText = FTextBlockStyle()
+			.SetFont(Consolas10)
+			.SetColorAndOpacity(FLinearColor::White)
+			.SetShadowOffset(FVector2D::ZeroVector)
+			.SetShadowColorAndOpacity(FLinearColor::Black)
+			.SetHighlightColor(FLinearColor(0.02f, 0.3f, 0.0f))
+			.SetHighlightShape(BOX_BRUSH("Images/TextBlockHighlightShape", FMargin(3.f / 8.f)));
+
+		const FVector2D Icon8x8(8.0f, 8.0f);
+		const FTextBlockStyle ErrorText = FTextBlockStyle(NormalText)
+			.SetUnderlineBrush(IMAGE_BRUSH("Images/White", Icon8x8, FLinearColor::Red, ESlateBrushTileType::Both))
+			.SetColorAndOpacity(FLinearColor::Red);
+		// ;
+		
+		Set("TextEditor.NormalText", NormalText);
+		
+		Set("TextEditor.Border", new BOX_BRUSH("Images/TextEditorBorder", FMargin(4.0f/16.0f), FLinearColor(0.02f,0.02f,0.02f,1)));
+
+		const FEditableTextBoxStyle EditableTextBoxStyle = FEditableTextBoxStyle()
+			.SetBackgroundImageNormal( FSlateNoResource() )
+			.SetBackgroundImageHovered( FSlateNoResource() )
+			.SetBackgroundImageFocused( FSlateNoResource() )
+			.SetBackgroundImageReadOnly( FSlateNoResource() );
+		
+		Set("TextEditor.EditableTextBox", EditableTextBoxStyle);
+	}
+}
diff --git a/Engine/Plugins/Experimental/Animation/CurveExpression/Source/Editor/Private/CurveExpressionEditorStyle.h b/Engine/Plugins/Experimental/Animation/CurveExpression/Source/Editor/Private/CurveExpressionEditorStyle.h
new file mode 100644
--- /dev/null
+++ b/Engine/Plugins/Experimental/Animation/CurveExpression/Source/Editor/Private/CurveExpressionEditorStyle.h
@@ -0,0 +1,18 @@
+// Copyright Epic Games, Inc. All Rights Reserved.
+
+#pragma once
+
+#include "Styling/SlateStyle.h"
+
+class FCurveExpressionEditorStyle
+    : public FSlateStyleSet
+{
+public:
+	static FCurveExpressionEditorStyle& Get();
+
+	static void Register();
+	static void Unregister();
+
+private:
+	FCurveExpressionEditorStyle();
+};
diff --git a/Engine/Plugins/Experimental/Animation/CurveExpression/Source/Editor/Private/K2Node_MakeCurveExpressionMap.cpp b/Engine/Plugins/Experimental/Animation/CurveExpression/Source/Editor/Private/K2Node_MakeCurveExpressionMap.cpp
new file mode 100644
--- /dev/null
+++ b/Engine/Plugins/Experimental/Animation/CurveExpression/Source/Editor/Private/K2Node_MakeCurveExpressionMap.cpp
@@ -0,0 +1,178 @@
+// Copyright Epic Games, Inc. All Rights Reserved.
+
+#include "K2Node_MakeCurveExpressionMap.h"
+
+#include "BlueprintActionDatabaseRegistrar.h"
+#include "BlueprintNodeSpawner.h"
+#include "Internationalization/Regex.h"
+#include "KismetCompiledFunctionContext.h"
+#include "KismetCompilerMisc.h"
+#include "String/ParseLines.h"
+#include "Textures/SlateIcon.h"
+
+#define LOCTEXT_NAMESPACE "MakeCurveExpressionMap"
+
+// ---------------------------------------------------------------------------------------------
+class FNodeHandler_MakeCurveExpressionMap final :
+	public FNodeHandlingFunctor
+{
+public:
+	FNodeHandler_MakeCurveExpressionMap(FKismetCompilerContext& InCompilerContext) :
+		FNodeHandlingFunctor(InCompilerContext)
+	{
+		
+	}
+	
+	virtual void RegisterNets(FKismetFunctionContext& InContext, UEdGraphNode* InNode) override
+	{
+		FNodeHandlingFunctor::RegisterNets(InContext, InNode);
+		
+		UEdGraphPin* OutputPin = CastChecked<UK2Node_MakeCurveExpressionMap>(InNode)->GetOutputPin();
+
+		FBPTerminal* Terminal = InContext.CreateLocalTerminalFromPinAutoChooseScope(OutputPin, InContext.NetNameMap->MakeValidName(OutputPin));
+		Terminal->bPassedByReference = false;
+		Terminal->Source = InNode;
+		InContext.NetMap.Add(OutputPin, Terminal);
+	}
+	
+	virtual void Compile(FKismetFunctionContext& InContext, UEdGraphNode* InNode) override
+	{
+		const UK2Node_MakeCurveExpressionMap* MapNode = CastChecked<UK2Node_MakeCurveExpressionMap>(InNode);
+		const UEdGraphPin* OutputPin = MapNode->GetOutputPin();
+
+		FBPTerminal** ContainerTerm = InContext.NetMap.Find(OutputPin);
+		if (!ensure(ContainerTerm))
+		{
+			return;
+		}
+
+		// Create a statement that assembles the map as a pile of literals.
+		FBlueprintCompiledStatement& CreateMapStatement = InContext.AppendStatementForNode(InNode);
+		CreateMapStatement.Type = KCST_CreateMap;
+		CreateMapStatement.LHS = *ContainerTerm;
+
+		FEdGraphPinType KeyType, ValueType;
+		KeyType.PinCategory = UEdGraphSchema_K2::PC_Name;
+		ValueType.PinCategory = UEdGraphSchema_K2::PC_String;
+		
+		for(TTuple<FName, FString>& Item: MapNode->GetExpressionMap())
+		{
+			FBPTerminal* KeyTerminal = new FBPTerminal();
+			KeyTerminal->Name = Item.Key.ToString();
+			KeyTerminal->Type = KeyType;
+			KeyTerminal->Source = InNode;
+			KeyTerminal->bIsLiteral = true;
+			InContext.Literals.Add(KeyTerminal);
+			
+			FBPTerminal* ValueTerminal = new FBPTerminal();
+			ValueTerminal->Name = MoveTemp(Item.Value);
+			ValueTerminal->Type = ValueType;
+			ValueTerminal->Source = InNode;
+			ValueTerminal->bIsLiteral = true;
+			InContext.Literals.Add(ValueTerminal);
+
+			CreateMapStatement.RHS.Add(KeyTerminal);
+			CreateMapStatement.RHS.Add(ValueTerminal);
+		}
+	}
+};
+
+
+// ---------------------------------------------------------------------------------------------
+
+const FName UK2Node_MakeCurveExpressionMap::OutputPinName(TEXT("Map"));
+
+
+UK2Node_MakeCurveExpressionMap::UK2Node_MakeCurveExpressionMap()
+{
+}
+
+
+UEdGraphPin* UK2Node_MakeCurveExpressionMap::GetOutputPin() const
+{
+	return FindPin(OutputPinName);
+}
+
+
+TMap<FName, FString> UK2Node_MakeCurveExpressionMap::GetExpressionMap() const
+{
+	TMap<FName, FString> ExpressionMap;
+	UE::String::ParseLines(Expressions.AssignmentExpressions, [&ExpressionMap](FStringView InLine)
+	{
+		int32 AssignmentPos;
+		if (InLine.FindChar('=', AssignmentPos))
+		{
+			FStringView Target = InLine.Left(AssignmentPos).TrimStartAndEnd();
+			FStringView Source = InLine.Mid(AssignmentPos + 1).TrimStartAndEnd();
+			if (!Target.IsEmpty() && !Source.IsEmpty())
+			{
+				ExpressionMap.Add(FName(Target), FString(Source));
+			}
+		}
+	});
+	
+	return ExpressionMap;
+}
+
+
+void UK2Node_MakeCurveExpressionMap::AllocateDefaultPins()
+{
+	FCreatePinParams PinParams;
+	PinParams.ContainerType = EPinContainerType::Map;
+	PinParams.ValueTerminalType.TerminalCategory = UEdGraphSchema_K2::PC_String;
+	
+	CreatePin(EGPD_Output, UEdGraphSchema_K2::PC_Name, OutputPinName, PinParams);
+}
+
+
+FText UK2Node_MakeCurveExpressionMap::GetNodeTitle(ENodeTitleType::Type TitleType) const
+{
+	return LOCTEXT("NodeTitle", "Make Expression Map");
+}
+
+
+FText UK2Node_MakeCurveExpressionMap::GetTooltipText() const
+{
+	return LOCTEXT("NodeTooltip", "Create an expression map from a list of assignment expressions");
+}
+
+
+FSlateIcon UK2Node_MakeCurveExpressionMap::GetIconAndTint(FLinearColor& OutColor) const
+{
+	static FSlateIcon Icon("EditorStyle", "GraphEditor.MakeMap_16x");
+	return Icon;
+}
+
+
+void UK2Node_MakeCurveExpressionMap::ValidateNodeDuringCompilation(FCompilerResultsLog& MessageLog) const
+{
+	// FIXME: Validate expressions.
+	Super::ValidateNodeDuringCompilation(MessageLog);
+}
+
+
+void UK2Node_MakeCurveExpressionMap::GetMenuActions(FBlueprintActionDatabaseRegistrar& InActionRegistrar) const
+{
+	const UClass* ActionKey = GetClass();
+	if (InActionRegistrar.IsOpenForRegistration(ActionKey))
+	{
+		UBlueprintNodeSpawner* NodeSpawner = UBlueprintNodeSpawner::Create(GetClass());
+		check(NodeSpawner != nullptr);
+
+		InActionRegistrar.AddBlueprintAction(ActionKey, NodeSpawner);
+	}
+}
+
+
+FNodeHandlingFunctor* UK2Node_MakeCurveExpressionMap::CreateNodeHandler(FKismetCompilerContext& CompilerContext) const
+{
+	return new FNodeHandler_MakeCurveExpressionMap(CompilerContext);
+}
+
+
+FText UK2Node_MakeCurveExpressionMap::GetMenuCategory() const
+{
+	return LOCTEXT("NodeCategory", "Curve Expression");
+}
+
+#undef LOCTEXT_NAMESPACE
diff --git a/Engine/Plugins/Experimental/Animation/CurveExpression/Source/Editor/Public/K2Node_MakeCurveExpressionMap.h b/Engine/Plugins/Experimental/Animation/CurveExpression/Source/Editor/Public/K2Node_MakeCurveExpressionMap.h
new file mode 100644
--- /dev/null
+++ b/Engine/Plugins/Experimental/Animation/CurveExpression/Source/Editor/Public/K2Node_MakeCurveExpressionMap.h
@@ -0,0 +1,51 @@
+// Copyright Epic Games, Inc. All Rights Reserved.
+
+#pragma once
+
+#include "K2Node.h"
+
+#include "K2Node_MakeCurveExpressionMap.generated.h"
+
+USTRUCT()
+struct FCurveExpressionList
+{
+	GENERATED_BODY()
+
+	UPROPERTY(EditAnywhere, Category="Expressions")
+	FString AssignmentExpressions;
+};
+
+UCLASS()
+class CURVEEXPRESSIONEDITOR_API UK2Node_MakeCurveExpressionMap :
+	public UK2Node
+{
+	GENERATED_BODY()
+
+public:
+	UK2Node_MakeCurveExpressionMap();
+	
+	UEdGraphPin* GetOutputPin() const;
+	TMap<FName, FString> GetExpressionMap() const;
+	
+	// UEdGraphNode interface
+	virtual void AllocateDefaultPins() override;
+	virtual FText GetNodeTitle(ENodeTitleType::Type TitleType) const override;
+	virtual FText GetTooltipText() const override;
+	virtual FSlateIcon GetIconAndTint(FLinearColor& OutColor) const override;
+	virtual void ValidateNodeDuringCompilation(FCompilerResultsLog& MessageLog) const override;
+	// End of UEdGraphNode interface
+
+	// UK2Node interface
+	virtual void GetMenuActions(FBlueprintActionDatabaseRegistrar& InActionRegistrar) const override;
+	virtual FNodeHandlingFunctor* CreateNodeHandler(FKismetCompilerContext& CompilerContext) const override;
+	virtual FText GetMenuCategory() const override;
+	virtual bool ShouldShowNodeProperties() const override { return true; }
+	virtual bool IsNodePure() const override { return true; }
+	// End of UK2Node interface
+
+	UPROPERTY(EditAnywhere, Category="Expressions")
+	FCurveExpressionList Expressions;
+
+private:
+	static const FName OutputPinName;
+};
