Curve Expressions: Add some rudimentary debugging for the curve expressions.

diff --git a/Engine/Plugins/Experimental/Animation/CurveExpression/Source/Editor/CurveExpressionEditor.Build.cs b/Engine/Plugins/Experimental/Animation/CurveExpression/Source/Editor/CurveExpressionEditor.Build.cs
--- a/Engine/Plugins/Experimental/Animation/CurveExpression/Source/Editor/CurveExpressionEditor.Build.cs
+++ b/Engine/Plugins/Experimental/Animation/CurveExpression/Source/Editor/CurveExpressionEditor.Build.cs
@@ -18,6 +18,7 @@ public class CurveExpressionEditor : ModuleRules
 				"KismetCompiler",
 				"Slate",
 				"SlateCore",
+				"UnrealEd",
 			}
 		);
 	}
diff --git a/Engine/Plugins/Experimental/Animation/CurveExpression/Source/Editor/Private/AnimGraphNode_RemapCurvesFromMesh.cpp b/Engine/Plugins/Experimental/Animation/CurveExpression/Source/Editor/Private/AnimGraphNode_RemapCurvesFromMesh.cpp
--- a/Engine/Plugins/Experimental/Animation/CurveExpression/Source/Editor/Private/AnimGraphNode_RemapCurvesFromMesh.cpp
+++ b/Engine/Plugins/Experimental/Animation/CurveExpression/Source/Editor/Private/AnimGraphNode_RemapCurvesFromMesh.cpp
@@ -5,6 +5,21 @@
 #define LOCTEXT_NAMESPACE "AnimGraphNode_RemapCurvesFromMesh"
 
 
+bool UAnimGraphNode_RemapCurvesFromMesh::CanVerifyExpressions() const
+{
+	return GetDebuggedNode() != nullptr;
+}
+
+
+void UAnimGraphNode_RemapCurvesFromMesh::VerifyExpressions()
+{
+	if (FAnimNode_RemapCurvesFromMesh* DebuggedNode = GetDebuggedNode())
+	{
+		DebuggedNode->VerifyExpressions();
+	}
+}
+
+
 FText UAnimGraphNode_RemapCurvesFromMesh::GetNodeTitle(ENodeTitleType::Type TitleType) const
 {
 	return LOCTEXT("NodeTitle", "Remap Curves From Mesh");
@@ -20,4 +35,22 @@ void UAnimGraphNode_RemapCurvesFromMesh::GetOutputLinkAttributes(FNodeAttributeA
 	OutAttributes.Add(UE::Anim::FAttributes::Curves);
 }
 
+
+FAnimNode_RemapCurvesFromMesh* UAnimGraphNode_RemapCurvesFromMesh::GetDebuggedNode() const
+{
+	if (const UObject* ObjectBeingDebugged = GetAnimBlueprint()->GetObjectBeingDebugged())
+	{
+		if (const UAnimInstance* InstanceBeingDebugged = Cast<const UAnimInstance>(ObjectBeingDebugged))
+		{
+			USkeletalMeshComponent* Component = InstanceBeingDebugged->GetSkelMeshComponent();
+			if (Component != nullptr && Component->GetAnimInstance() != nullptr)
+			{
+				return static_cast<FAnimNode_RemapCurvesFromMesh*>(FindDebugAnimNode(Component));
+			}
+		}
+	}
+
+	return nullptr;
+}
+
 #undef LOCTEXT_NAMESPACE
diff --git a/Engine/Plugins/Experimental/Animation/CurveExpression/Source/Editor/Private/CurveExpressionDetailsCustomization.cpp b/Engine/Plugins/Experimental/Animation/CurveExpression/Source/Editor/Private/CurveExpressionDetailsCustomization.cpp
--- a/Engine/Plugins/Experimental/Animation/CurveExpression/Source/Editor/Private/CurveExpressionDetailsCustomization.cpp
+++ b/Engine/Plugins/Experimental/Animation/CurveExpression/Source/Editor/Private/CurveExpressionDetailsCustomization.cpp
@@ -2,14 +2,82 @@
 
 #include "CurveExpressionDetailsCustomization.h"
 
+#include "AnimGraphNode_RemapCurvesFromMesh.h"
 #include "CurveExpressionEditorStyle.h"
 #include "K2Node_MakeCurveExpressionMap.h"
+
+#include "DetailLayoutBuilder.h"
 #include "DetailWidgetRow.h"
+#include "Widgets/Input/SButton.h"
 #include "Widgets/Layout/SGridPanel.h"
 #include "Widgets/Layout/SScrollBox.h"
 #include "Widgets/Text/SMultiLineEditableText.h"
 
 
+#define LOCTEXT_NAMESPACE "CurveExpressionDetailsCustomization"
+
+TSharedRef<IDetailCustomization> FAnimGraphNode_RemapCurvesFromMeshCustomization::MakeInstance()
+{
+	return MakeShared<FAnimGraphNode_RemapCurvesFromMeshCustomization>();
+}
+
+
+void FAnimGraphNode_RemapCurvesFromMeshCustomization::CustomizeDetails(
+	IDetailLayoutBuilder& InDetailBuilder
+	)
+{
+	TArray<TWeakObjectPtr<UObject>> ObjectsBeingCustomized;
+	InDetailBuilder.GetObjectsBeingCustomized(ObjectsBeingCustomized);
+
+	TArray<TWeakObjectPtr<UAnimGraphNode_RemapCurvesFromMesh>> RemapNodes;
+	for (TWeakObjectPtr<UObject> Object: ObjectsBeingCustomized)
+	{
+		TWeakObjectPtr<UAnimGraphNode_RemapCurvesFromMesh> RemapNode = Cast<UAnimGraphNode_RemapCurvesFromMesh>(Object.Get());
+		if (RemapNode.IsValid())
+		{
+			RemapNodes.Add(RemapNode);
+		}
+	}
+
+	IDetailCategoryBuilder& DebuggingCategory = InDetailBuilder.EditCategory("Debugging");
+	DebuggingCategory.AddCustomRow(LOCTEXT("Debugging", "Debugging"))
+		.ValueContent()
+		[
+			SNew(SHorizontalBox)
+			+ SHorizontalBox::Slot()
+			.Padding(5)
+			.AutoWidth()
+			[
+				SNew(SButton)
+				.Text(LOCTEXT("VerifyExpressions", "Verify Expressions"))
+				.ToolTipText(LOCTEXT("VerifyExpressions_Tooltip", "Verify all the expressions as applied to the currently debugged object. If debugging is not active, no verification is done. Verification output is sent to the output log."))
+				.OnClicked_Lambda([this, RemapNodes]()
+				{
+					for (TWeakObjectPtr<UAnimGraphNode_RemapCurvesFromMesh> RemapNodeWeak: RemapNodes)
+					{
+						if (UAnimGraphNode_RemapCurvesFromMesh* RemapNode = RemapNodeWeak.Get())
+						{
+							RemapNode->VerifyExpressions();
+						}
+					}
+					return FReply::Handled();
+				})
+				.IsEnabled_Lambda([this, RemapNodes]()
+				{
+					for (TWeakObjectPtr<UAnimGraphNode_RemapCurvesFromMesh> RemapNodeWeak: RemapNodes)
+					{
+						if (UAnimGraphNode_RemapCurvesFromMesh* RemapNode = RemapNodeWeak.Get())
+						{
+							return true;
+						}
+					}
+					return false;
+				})
+			]
+		];
+}
+
+
 TSharedRef<IPropertyTypeCustomization> FCurveExpressionListCustomization::MakeInstance()
 {
 	return MakeShared<FCurveExpressionListCustomization>();
@@ -96,3 +164,5 @@ void FCurveExpressionListCustomization::CustomizeHeader(
 		]
 	];
 }
+
+#undef LOCTEXT_NAMESPACE
diff --git a/Engine/Plugins/Experimental/Animation/CurveExpression/Source/Editor/Private/CurveExpressionDetailsCustomization.h b/Engine/Plugins/Experimental/Animation/CurveExpression/Source/Editor/Private/CurveExpressionDetailsCustomization.h
--- a/Engine/Plugins/Experimental/Animation/CurveExpression/Source/Editor/Private/CurveExpressionDetailsCustomization.h
+++ b/Engine/Plugins/Experimental/Animation/CurveExpression/Source/Editor/Private/CurveExpressionDetailsCustomization.h
@@ -2,6 +2,7 @@
 
 #pragma once
 
+#include "IDetailCustomization.h"
 #include "IPropertyTypeCustomization.h"
 
 
@@ -9,6 +10,17 @@ class SMultiLineEditableText;
 class SScrollBar;
 
 
+class FAnimGraphNode_RemapCurvesFromMeshCustomization :
+	public IDetailCustomization
+{
+public:
+	static TSharedRef<IDetailCustomization> MakeInstance();
+
+	// IDetailCustomization overrides
+	virtual void CustomizeDetails(IDetailLayoutBuilder& InDetailBuilder) override;
+};
+
+
 class FCurveExpressionListCustomization : 
 	public IPropertyTypeCustomization
 {
@@ -34,3 +46,4 @@ private:
 
 	TSharedPtr<SMultiLineEditableText> TextEditor;	
 };
+
diff --git a/Engine/Plugins/Experimental/Animation/CurveExpression/Source/Editor/Private/CurveExpressionEditorModule.cpp b/Engine/Plugins/Experimental/Animation/CurveExpression/Source/Editor/Private/CurveExpressionEditorModule.cpp
--- a/Engine/Plugins/Experimental/Animation/CurveExpression/Source/Editor/Private/CurveExpressionEditorModule.cpp
+++ b/Engine/Plugins/Experimental/Animation/CurveExpression/Source/Editor/Private/CurveExpressionEditorModule.cpp
@@ -1,5 +1,6 @@
 // Copyright Epic Games, Inc. All Rights Reserved.
 
+#include "AnimGraphNode_RemapCurvesFromMesh.h"
 #include "CurveExpressionDetailsCustomization.h"
 #include "CurveExpressionEditorStyle.h"
 #include "K2Node_MakeCurveExpressionMap.h"
@@ -23,6 +24,14 @@ public:
 				);
 			CustomizedProperties.Add(InStructName);
 		};
+		auto RegisterClassCustomization = [&](FName InClassName, auto InCustomizationFactory)
+		{
+			PropertyModule.RegisterCustomClassLayout(
+				InClassName, 
+				FOnGetDetailCustomizationInstance::CreateStatic(InCustomizationFactory)
+				);
+			CustomizedClasses.Add(InClassName);
+		};
 		
 		FCurveExpressionEditorStyle::Register();
 		
@@ -30,6 +39,11 @@ public:
 			FCurveExpressionList::StaticStruct()->GetFName(),
 			&FCurveExpressionListCustomization::MakeInstance
 			);
+			
+		RegisterClassCustomization(
+			UAnimGraphNode_RemapCurvesFromMesh::StaticClass()->GetFName(),
+			&FAnimGraphNode_RemapCurvesFromMeshCustomization::MakeInstance
+			);
 	}
 	virtual void ShutdownModule() override
 	{
@@ -39,6 +53,10 @@ public:
 			{
 				PropertyModule->UnregisterCustomPropertyTypeLayout(PropertyName);
 			}
+			for (const FName& ClassName: CustomizedClasses)
+			{
+				PropertyModule->UnregisterCustomClassLayout(ClassName);
+			}
 		}
 		
 		FCurveExpressionEditorStyle::Unregister();
@@ -46,6 +64,7 @@ public:
 
 private:
 	TArray<FName> CustomizedProperties;
+	TArray<FName> CustomizedClasses;
 };
 
 }
diff --git a/Engine/Plugins/Experimental/Animation/CurveExpression/Source/Editor/Private/K2Node_MakeCurveExpressionMap.cpp b/Engine/Plugins/Experimental/Animation/CurveExpression/Source/Editor/Private/K2Node_MakeCurveExpressionMap.cpp
--- a/Engine/Plugins/Experimental/Animation/CurveExpression/Source/Editor/Private/K2Node_MakeCurveExpressionMap.cpp
+++ b/Engine/Plugins/Experimental/Animation/CurveExpression/Source/Editor/Private/K2Node_MakeCurveExpressionMap.cpp
@@ -2,9 +2,10 @@
 
 #include "K2Node_MakeCurveExpressionMap.h"
 
+#include "ExpressionEvaluator.h"
+
 #include "BlueprintActionDatabaseRegistrar.h"
 #include "BlueprintNodeSpawner.h"
-#include "Internationalization/Regex.h"
 #include "KismetCompiledFunctionContext.h"
 #include "KismetCompilerMisc.h"
 #include "String/ParseLines.h"
@@ -97,16 +98,17 @@ UEdGraphPin* UK2Node_MakeCurveExpressionMap::GetOutputPin() const
 TMap<FName, FString> UK2Node_MakeCurveExpressionMap::GetExpressionMap() const
 {
 	TMap<FName, FString> ExpressionMap;
-	UE::String::ParseLines(Expressions.AssignmentExpressions, [&ExpressionMap](FStringView InLine)
+	UE::String::ParseLines(Expressions.AssignmentExpressions,
+		[&ExpressionMap](FStringView InLine)
 	{
 		int32 AssignmentPos;
 		if (InLine.FindChar('=', AssignmentPos))
 		{
-			FStringView Target = InLine.Left(AssignmentPos).TrimStartAndEnd();
-			FStringView Source = InLine.Mid(AssignmentPos + 1).TrimStartAndEnd();
-			if (!Target.IsEmpty() && !Source.IsEmpty())
+			FStringView TargetCurve = InLine.Left(AssignmentPos).TrimStartAndEnd();
+			FStringView SourceExpression = InLine.Mid(AssignmentPos + 1).TrimStartAndEnd();
+			if (!TargetCurve.IsEmpty() && !SourceExpression.IsEmpty())
 			{
-				ExpressionMap.Add(FName(Target), FString(Source));
+				ExpressionMap.Add(FName(TargetCurve), FString(SourceExpression));
 			}
 		}
 	});
@@ -146,8 +148,70 @@ FSlateIcon UK2Node_MakeCurveExpressionMap::GetIconAndTint(FLinearColor& OutColor
 
 void UK2Node_MakeCurveExpressionMap::ValidateNodeDuringCompilation(FCompilerResultsLog& MessageLog) const
 {
-	// FIXME: Validate expressions.
 	Super::ValidateNodeDuringCompilation(MessageLog);
+
+	using namespace CurveExpression::Evaluator;
+
+	// Do a check on the expression form. We ignore curve names for this.
+	static FEngine Engine;
+
+	int32 LineNumber = 1;
+	UE::String::ParseLines(Expressions.AssignmentExpressions,
+		[&](FStringView InLine)
+	{
+		FFormatNamedArguments Args;
+		Args.Add(TEXT("Line"), FText::AsNumber(LineNumber));
+
+		if (InLine.TrimStartAndEnd().IsEmpty())
+		{
+			return;
+		}
+		
+		int32 AssignmentPos;
+		if (InLine.FindChar('=', AssignmentPos))
+		{
+			FStringView Target = InLine.Left(AssignmentPos).TrimStartAndEnd();
+			FStringView SourceExpression = InLine.Mid(AssignmentPos + 1).TrimStartAndEnd();
+
+			if (Target.IsEmpty())
+			{
+				MessageLog.Error(
+					*FText::Format(LOCTEXT("Warning_NoTarget", "@@ has an error on line {Line}. Target curve not specified."), Args).ToString(), 
+					this
+				);
+			}
+			else if (SourceExpression.IsEmpty())
+			{
+				MessageLog.Error(
+					*FText::Format(LOCTEXT("Warning_NoSource", "@@ has an error on line {Line}. No source expression specified."), Args).ToString(), 
+					this
+				);
+			}
+			else
+			{
+				TOptional<FParseError> Error = Engine.Verify(SourceExpression);
+				if (Error.IsSet())
+				{
+					const int32 ExpressionColumn = SourceExpression.GetData() - InLine.GetData();
+					Args.Add(TEXT("Error"), FText::FromString(Error->Message));
+					Args.Add(TEXT("Column"), FText::AsNumber(Error->Column + ExpressionColumn + 1));
+				
+					MessageLog.Error(
+						*FText::Format(LOCTEXT("Warning_BadExpression", "@@ has an error on line {Line}, column {Column}. {Error}"), Args).ToString(), 
+						this
+					);
+				}
+			}
+		}
+		else
+		{
+			MessageLog.Error(
+				*FText::Format(LOCTEXT("Warning_NoAssignment", "@@ has an error on line {Line}. Not an assignment (e.g. 'TargetCurve = SourceCurve`)."), Args).ToString(), 
+				this
+			);
+		}
+		LineNumber++;
+	});
 }
 
 
diff --git a/Engine/Plugins/Experimental/Animation/CurveExpression/Source/Editor/Public/AnimGraphNode_RemapCurvesFromMesh.h b/Engine/Plugins/Experimental/Animation/CurveExpression/Source/Editor/Public/AnimGraphNode_RemapCurvesFromMesh.h
--- a/Engine/Plugins/Experimental/Animation/CurveExpression/Source/Editor/Public/AnimGraphNode_RemapCurvesFromMesh.h
+++ b/Engine/Plugins/Experimental/Animation/CurveExpression/Source/Editor/Public/AnimGraphNode_RemapCurvesFromMesh.h
@@ -15,9 +15,13 @@ class UAnimGraphNode_RemapCurvesFromMesh :
 {
 	GENERATED_BODY()
 
+public:
 	UPROPERTY(EditAnywhere, Category=Settings)
 	FAnimNode_RemapCurvesFromMesh Node;
 
+	bool CanVerifyExpressions() const;
+	void VerifyExpressions();
+
 	// UEdGraphNode interface
 	virtual FText GetNodeTitle(ENodeTitleType::Type TitleType) const override;
 	virtual FText GetTooltipText() const override;
@@ -27,4 +31,8 @@ class UAnimGraphNode_RemapCurvesFromMesh :
 	virtual bool UsingCopyPoseFromMesh() const override { return true; }
 	virtual void GetOutputLinkAttributes(FNodeAttributeArray& OutAttributes) const override;
 	// End of UAnimGraphNode_Base interface
+	
+private:
+	FAnimNode_RemapCurvesFromMesh* GetDebuggedNode() const;
+	
 };
diff --git a/Engine/Plugins/Experimental/Animation/CurveExpression/Source/Runtime/Private/AnimNode_RemapCurvesFromMesh.cpp b/Engine/Plugins/Experimental/Animation/CurveExpression/Source/Runtime/Private/AnimNode_RemapCurvesFromMesh.cpp
--- a/Engine/Plugins/Experimental/Animation/CurveExpression/Source/Runtime/Private/AnimNode_RemapCurvesFromMesh.cpp
+++ b/Engine/Plugins/Experimental/Animation/CurveExpression/Source/Runtime/Private/AnimNode_RemapCurvesFromMesh.cpp
@@ -2,6 +2,7 @@
 
 #include "AnimNode_RemapCurvesFromMesh.h"
 
+#include "CurveExpressionModule.h"
 #include "ExpressionEvaluator.h"
 
 void FAnimNode_RemapCurvesFromMesh::Initialize_AnyThread(
@@ -24,6 +25,35 @@ void FAnimNode_RemapCurvesFromMesh::CacheBones_AnyThread(
 }
 	
 
+void FAnimNode_RemapCurvesFromMesh::VerifyExpressions()
+{
+	using namespace CurveExpression::Evaluator;
+
+	if (!ExpressionEngine.IsSet() || CurveExpressions.IsEmpty())
+	{
+		return;
+	}
+
+	FEngine VerificationEngine(
+		ExpressionEngine->GetConstantValues(),
+		EParseFlags::ValidateConstants);
+
+	for (const TPair<FName, FString>& ExpressionPair: CurveExpressions)
+	{
+		if (!CurveNameToUIDMap.Contains(ExpressionPair.Key))
+		{
+			UE_LOG(LogCurveExpression, Warning, TEXT("Target curve '%s' does not exist."), *ExpressionPair.Key.ToString());
+		}
+
+		TOptional<FParseError> Error = VerificationEngine.Verify(ExpressionPair.Value);
+		if (Error.IsSet())
+		{
+			UE_LOG(LogCurveExpression, Warning, TEXT("Expression error in '%s': %s"), *ExpressionPair.Value, *Error->Message);
+		}
+	}
+}
+
+
 void FAnimNode_RemapCurvesFromMesh::Update_AnyThread(
 	const FAnimationUpdateContext& Context
 	)
@@ -50,7 +80,6 @@ void FAnimNode_RemapCurvesFromMesh::Update_AnyThread(
 				{
 					CachedExpressions.Add(ExpressionPair.Key, *Expression);
 				}
-				// Parse errors we deal with in the animgraph node.
 			}
 		}
 	}			
@@ -181,19 +210,16 @@ void FAnimNode_RemapCurvesFromMesh::ReinitializeMeshComponent(
 			CurrentlyUsedTargetMesh = TargetSkelMesh;
 
 			// Grab the source curves and use their names to seed the expression evaluation engine.
+			TMap<FName, float> SourceCurves;
 			const USkeleton* SourceSkeleton = SourceSkelMesh->GetSkeleton();
 			if (ensureMsgf(SourceSkeleton, TEXT("Invalid null source skeleton : %s"), *GetNameSafe(TargetSkelMesh)))
 			{
 				const FSmartNameMapping* SourceContainer = SourceSkeleton->GetSmartNameContainer(USkeleton::AnimCurveMappingName);
 
-				TMap<FName, float> SourceCurves;
-				
 				SourceContainer->Iterate([&SourceCurves](const FSmartNameMappingIterator& Iterator)
 				{
 					if (FName CurveName; Iterator.GetName(CurveName))
 					{
-						const FCurveMetaData* MetaData = Iterator.GetCurveMetaData();
-						FPlatformMisc::LowLevelOutputDebugStringf(TEXT("CURVE[%s] - Morph: %d - Mat: %d\n"), *CurveName.ToString(), MetaData->Type.bMaterial, MetaData->Type.bMorphtarget);
 						SourceCurves.Add(CurveName, 0.0f);
 					}
 				});
diff --git a/Engine/Plugins/Experimental/Animation/CurveExpression/Source/Runtime/Private/CurveExpressionModule.cpp b/Engine/Plugins/Experimental/Animation/CurveExpression/Source/Runtime/Private/CurveExpressionModule.cpp
--- a/Engine/Plugins/Experimental/Animation/CurveExpression/Source/Runtime/Private/CurveExpressionModule.cpp
+++ b/Engine/Plugins/Experimental/Animation/CurveExpression/Source/Runtime/Private/CurveExpressionModule.cpp
@@ -1,5 +1,6 @@
 // Copyright Epic Games, Inc. All Rights Reserved.
 
+#include "CurveExpressionModule.h"
 #include "Modules/ModuleManager.h"
 
 namespace UE::CurveExpression
@@ -14,4 +15,6 @@ public:
 
 }
 
+DEFINE_LOG_CATEGORY(LogCurveExpression);
+
 IMPLEMENT_MODULE(UE::CurveExpression::FModule, CurveExpression)
diff --git a/Engine/Plugins/Experimental/Animation/CurveExpression/Source/Runtime/Private/ExpressionEvaluator.cpp b/Engine/Plugins/Experimental/Animation/CurveExpression/Source/Runtime/Private/ExpressionEvaluator.cpp
--- a/Engine/Plugins/Experimental/Animation/CurveExpression/Source/Runtime/Private/ExpressionEvaluator.cpp
+++ b/Engine/Plugins/Experimental/Animation/CurveExpression/Source/Runtime/Private/ExpressionEvaluator.cpp
@@ -9,9 +9,21 @@ namespace CurveExpression::Evaluator
 {
 
 FEngine::FEngine(
-	TMap<FName, float>&& InConstants
+	const TMap<FName, float>& InConstants, 
+	EParseFlags InParseFlags
 	) :
-	Constants(MoveTemp(InConstants))
+	Constants(InConstants),
+	ParseFlags(InParseFlags)
+{
+}
+
+
+FEngine::FEngine(
+	TMap<FName, float>&& InConstants,
+	EParseFlags InParseFlags
+	) :
+	Constants(MoveTemp(InConstants)),
+	ParseFlags(InParseFlags)
 {
 }
 
@@ -331,10 +343,14 @@ TVariant<FEngine::FToken, FParseError> FEngine::ParseIdentifier(
 		return TVariant<FToken, FParseError>(
 			TInPlaceType<FToken>(), FToken(TInPlaceType<FName>(), Name));
 	}
+	else if (EnumHasAllFlags(ParseFlags, EParseFlags::ValidateConstants))
+	{
+		return ParseTokenError(FString::Printf(TEXT("Unknown identifier '%s'"), *FString(Identifier)),
+			InOutParseRange.GetData(), InExpression); 
+	}
 	else
 	{
 		// If the constant value is not found, it gets a default value of 0.
-		// TBD: We may want to flag this as an error instead when doing verification.
 		return TVariant<FToken, FParseError>(
 			TInPlaceType<FToken>(), FToken(TInPlaceType<float>(), 0.0f));
 	}
diff --git a/Engine/Plugins/Experimental/Animation/CurveExpression/Source/Runtime/Public/AnimNode_RemapCurvesFromMesh.h b/Engine/Plugins/Experimental/Animation/CurveExpression/Source/Runtime/Public/AnimNode_RemapCurvesFromMesh.h
--- a/Engine/Plugins/Experimental/Animation/CurveExpression/Source/Runtime/Public/AnimNode_RemapCurvesFromMesh.h
+++ b/Engine/Plugins/Experimental/Animation/CurveExpression/Source/Runtime/Public/AnimNode_RemapCurvesFromMesh.h
@@ -30,6 +30,9 @@ struct CURVEEXPRESSION_API FAnimNode_RemapCurvesFromMesh :
 	/** The expressions given are immutable and will not change during runtime. Improves performance. */ 
 	UPROPERTY(EditAnywhere, BlueprintReadWrite, Category=Expressions, meta=(NeverAsPin))
 	bool bExpressionsImmutable = true;
+
+	// Call to verify the expressions. Report any errors back.
+	void VerifyExpressions();
 	
 	// FAnimNode_Base interface
 	virtual void Initialize_AnyThread(const FAnimationInitializeContext& Context) override;
diff --git a/Engine/Plugins/Experimental/Animation/CurveExpression/Source/Runtime/Public/CurveExpressionModule.h b/Engine/Plugins/Experimental/Animation/CurveExpression/Source/Runtime/Public/CurveExpressionModule.h
new file mode 100644
--- /dev/null
+++ b/Engine/Plugins/Experimental/Animation/CurveExpression/Source/Runtime/Public/CurveExpressionModule.h
@@ -0,0 +1,8 @@
+// Copyright Epic Games, Inc. All Rights Reserved.
+
+#pragma once
+
+#include "Logging/LogMacros.h"
+
+
+DECLARE_LOG_CATEGORY_EXTERN(LogCurveExpression, Log, All);
diff --git a/Engine/Plugins/Experimental/Animation/CurveExpression/Source/Runtime/Public/ExpressionEvaluator.h b/Engine/Plugins/Experimental/Animation/CurveExpression/Source/Runtime/Public/ExpressionEvaluator.h
--- a/Engine/Plugins/Experimental/Animation/CurveExpression/Source/Runtime/Public/ExpressionEvaluator.h
+++ b/Engine/Plugins/Experimental/Animation/CurveExpression/Source/Runtime/Public/ExpressionEvaluator.h
@@ -66,14 +66,31 @@ private:
  *	rather than:
  *	   foo (1) * 2.0
  */
-class FEngine
+enum class EParseFlags
+{
+	Default				= 0,		/** Default parsing behavior */
+	ValidateConstants	= 1 << 0,	/** Validate that the referred-to constant exists */
+};
+ENUM_CLASS_FLAGS(EParseFlags)
+
+
+class CURVEEXPRESSION_API FEngine
 {
 public:
-	FEngine() = default;
+	explicit FEngine(EParseFlags InParseFlags = EParseFlags::Default) :
+		ParseFlags(InParseFlags)
+	{ }
+
+	/** Initialize the engine with constants, making a copy of the value constant map. */
+	explicit FEngine(
+		const TMap<FName, float>& InConstants,
+		EParseFlags InParseFlags = EParseFlags::Default
+		);
 	
-	/** Initialize the engine with constants */
+	/** Initialize the engine with constants, taking ownership of the value constant map. */
 	explicit FEngine(
-		TMap<FName, float>&& InConstants
+		TMap<FName, float>&& InConstants,
+		EParseFlags InParseFlags = EParseFlags::Default
 		);
 
 	/** Update existing constant values. No new constants are added. */
@@ -81,6 +98,12 @@ public:
 		const TMap<FName, float>& InConstants
 		);
 
+	/** Return existing constants and their values. */
+	const TMap<FName, float>& GetConstantValues() const
+	{
+		return Constants;
+	}
+
 	/** Parse an expression and either receive an error or a expression execution object that
 	 *  matches this evaluator. */
 	TVariant<FExpressionObject, FParseError> Parse(
@@ -164,6 +187,9 @@ private:
 	// FUTURE: Add support for variables by allowing TFunction{Ref} objects along with direct
 	// constants.
 	TMap<FName, float> Constants;
+
+	/** The parse flags for this engine. Only set once */
+	EParseFlags ParseFlags = EParseFlags::Default;
 };
 
 }
