Curve Expressions: Fix include gone missing.

#rnx

diff --git a/Engine/Plugins/Experimental/Animation/CurveExpression/Source/Editor/Private/AnimGraphNode_RemapCurvesFromMesh.cpp b/Engine/Plugins/Experimental/Animation/CurveExpression/Source/Editor/Private/AnimGraphNode_RemapCurvesFromMesh.cpp
--- a/Engine/Plugins/Experimental/Animation/CurveExpression/Source/Editor/Private/AnimGraphNode_RemapCurvesFromMesh.cpp
+++ b/Engine/Plugins/Experimental/Animation/CurveExpression/Source/Editor/Private/AnimGraphNode_RemapCurvesFromMesh.cpp
@@ -2,6 +2,8 @@
 
 #include "AnimGraphNode_RemapCurvesFromMesh.h"
 
+#include "Animation/AnimAttributes.h"
+
 #define LOCTEXT_NAMESPACE "AnimGraphNode_RemapCurvesFromMesh"
 
 
