Curve Expressions: Added support for functions. Current set of builtins are: clamp, min, max, abs, round, floor, ceil, sin, cos, tan, asin, acos, atan, sqrt, isqrt, log, exp, pi and e.

diff --git a/Engine/Plugins/Experimental/Animation/CurveExpression/Source/Runtime/Private/ExpressionEvaluator.cpp b/Engine/Plugins/Experimental/Animation/CurveExpression/Source/Runtime/Private/ExpressionEvaluator.cpp
--- a/Engine/Plugins/Experimental/Animation/CurveExpression/Source/Runtime/Private/ExpressionEvaluator.cpp
+++ b/Engine/Plugins/Experimental/Animation/CurveExpression/Source/Runtime/Private/ExpressionEvaluator.cpp
@@ -1,6 +1,5 @@
 // Copyright Epic Games, Inc. All Rights Reserved.
 
-
 #include "ExpressionEvaluator.h"
 
 #include <limits>
@@ -8,6 +7,104 @@
 namespace CurveExpression::Evaluator
 {
 
+#define CE_EXPR(_N_, _A_, _E_) { \
+	FEngine::FunctionNameIndex.Add(FName(#_N_), FEngine::Functions.Num()); \
+	FEngine::Functions.Add({_A_, [](TArrayView<const float> V) -> float { return _E_; } }); \
+	}  
+
+TArray<FEngine::FFunctionInfo> FEngine::Functions;
+TMap<FName, int32> FEngine::FunctionNameIndex;
+
+struct FInitializeBuiltinFunctions
+{
+	FInitializeBuiltinFunctions()
+	{
+		// clamp(value, min, max)
+		CE_EXPR(clamp, 3, FMath::Clamp(V[0], V[1], V[2]))
+		// min(a, b)
+		CE_EXPR(min, 2, FMath::Min(V[0], V[1]))
+		// max(a, b)
+		CE_EXPR(max, 2, FMath::Max(V[0], V[1]))
+		
+		// abs(value)
+		CE_EXPR(abs, 1, FMath::Abs(V[0]))
+		// round(value)
+		CE_EXPR(round, 1, FMath::RoundToFloat(V[0]))
+		// ceil(value)
+		CE_EXPR(ceil, 1, FMath::CeilToFloat(V[0]))
+		// floor(value)
+		CE_EXPR(floor, 1, FMath::FloorToFloat(V[0]))
+
+		// sin(value)
+		CE_EXPR(sin, 1, FMath::Sin(V[0]))
+		// cos(value)
+		CE_EXPR(cos, 1, FMath::Cos(V[0]))
+		// tan(value)
+		CE_EXPR(tan, 1, FMath::Tan(V[0]))
+		// asin(value)
+		CE_EXPR(asin, 1, FMath::Asin(V[0]))
+		// acos(value)
+		CE_EXPR(acos, 1, FMath::Acos(V[0]))
+		// atan(value)
+		CE_EXPR(atan, 1, FMath::Atan(V[0]))
+
+		// sqrt(value)
+		CE_EXPR(sqrt, 1, FMath::Sqrt(V[0]))
+		// isqrt(value)
+		CE_EXPR(isqrt, 1, FMath::InvSqrt(V[0]))
+
+		// log(value)
+		CE_EXPR(log, 1, FMath::Loge(V[0]))
+		// exp(value)
+		CE_EXPR(exp, 1, FMath::Exp(V[0]))
+	
+		/** pi() */
+		CE_EXPR(pi, 0, UE_PI)
+		
+		/** pi() */
+		CE_EXPR(e, 0, UE_EULERS_NUMBER)
+	}
+} GInitializeBuiltinFunctions;
+#undef CE_EXPR
+
+
+FString FExpressionObject::ToString() const
+{
+	TArray<FString> Items;
+	for (const OpElement& Element: Expression)
+	{
+		if (const EOperator* Op = Element.TryGet<EOperator>())
+		{
+			switch(*Op)
+			{
+			case EOperator::Negate: Items.Add(TEXT("Op[Negate]")); break;
+			case EOperator::Add: Items.Add(TEXT("Op[Add]")); break;
+			case EOperator::Subtract: Items.Add(TEXT("Op[Subtract]")); break; 
+			case EOperator::Multiply: Items.Add(TEXT("Op[Multiply]")); break;
+			case EOperator::Divide: Items.Add(TEXT("Op[Divide]")); break;
+			case EOperator::Modulo: Items.Add(TEXT("Op[Modulo]")); break;
+			case EOperator::Power: Items.Add(TEXT("Op[Power]")); break;
+			case EOperator::FloorDivide: Items.Add(TEXT("Op[FloorDivide]")); break;
+			}
+		}
+		else if (const FName* Constant = Element.TryGet<FName>())
+		{
+			Items.Add(FString::Printf(TEXT("C[%s]"), *Constant->ToString()));
+		}
+		else if (const FFunctionRef *FuncRef = Element.TryGet<FFunctionRef>())
+		{
+			Items.Add(FString::Printf(TEXT("F[%d]"), FuncRef->Index));
+		}
+		else if (const float *Value = Element.TryGet<float>())
+		{
+			Items.Add(FString::Printf(TEXT("V[%g]"), *Value));
+		}
+	}
+
+	return FString::Join(Items, TEXT("\n"));
+}
+
+
 FEngine::FEngine(
 	const TMap<FName, float>& InConstants, 
 	EParseFlags InParseFlags
@@ -80,6 +177,7 @@ FString FEngine::TokenToString(
 		case EOperatorToken::FloorDivide:	ValueStr = TEXT("//"); break;
 		case EOperatorToken::ParenOpen:		ValueStr = TEXT("("); break;
 		case EOperatorToken::ParenClose:	ValueStr = TEXT(")"); break;
+		case EOperatorToken::Comma:			ValueStr = TEXT(","); break;
 		}
 	}
 	else if (const FName* Identifier = InToken.TryGet<FName>())
@@ -338,7 +436,7 @@ TVariant<FEngine::FToken, FParseError> FEngine::ParseIdentifier(
 	
 	InOutParseRange = InOutParseRange.Mid(ParseIndex); 
 
-	if (Constants.Contains(Name))
+	if (Constants.Contains(Name) || FunctionNameIndex.Contains(Name))
 	{
 		return TVariant<FToken, FParseError>(
 			TInPlaceType<FToken>(), FToken(TInPlaceType<FName>(), Name));
@@ -350,9 +448,9 @@ TVariant<FEngine::FToken, FParseError> FEngine::ParseIdentifier(
 	}
 	else
 	{
-		// If the constant value is not found, it gets a default value of 0.
+		// If the constant value is not found, it gets a default value of None.
 		return TVariant<FToken, FParseError>(
-			TInPlaceType<FToken>(), FToken(TInPlaceType<float>(), 0.0f));
+			TInPlaceType<FToken>(), FToken(TInPlaceType<FName>(), NAME_None));
 	}
 }
 
@@ -407,6 +505,10 @@ TVariant<FEngine::FToken, FParseError> FEngine::ParseToken(
 	case ')':
 		InOutParseRange = InOutParseRange.Mid(1);
 		return TokenResult(EOperatorToken::ParenClose);
+
+	case ',':
+		InOutParseRange = InOutParseRange.Mid(1);
+		return TokenResult(EOperatorToken::Comma);
 		
 	default:
 		if (IsASCIIDigit(InOutParseRange[0]) || InOutParseRange[0] == '.')
@@ -463,10 +565,18 @@ TVariant<FExpressionObject, FParseError> FEngine::Parse(
 	}
 
 	// Run Dijkstra's classic Shunting Yard algorithm to convert infix expressions to RPN.
-	// TODO: Add support for calling functions.
 	// TODO: Better error pinpointing by storing the token parse result with the token.
 	TArray<FExpressionObject::OpElement, TInlineAllocator<64>> Expression;
 	TArray<EOperatorToken, TInlineAllocator<32>> OperatorStack;
+	struct FFunctionCallInfo
+	{
+		int32 FunctionIndex = INDEX_NONE;
+		int32 CountedCommas = 0;
+		int32 OpeningOpStackSize = 0;
+		int32 ExpressionSize = 0;
+		const TCHAR* TokenStart = nullptr;
+	};
+	TArray<FFunctionCallInfo, TInlineAllocator<32>> FunctionStack;
 
 	auto ParseError = [InExpression](FString&& InError, const TCHAR* InCurrentPos)
 	{
@@ -481,6 +591,7 @@ TVariant<FExpressionObject, FParseError> FEngine::Parse(
 	// Used to check if we encounter a '-' to see if we need to turn it into a negate operator
 	// and also to see if we have two operators in a row (error).
 	TOptional<FToken> LastToken;
+	FStringView LastTokenParseRange;
 
 	do
 	{
@@ -533,7 +644,8 @@ TVariant<FExpressionObject, FParseError> FEngine::Parse(
 					{
 						return ParseError(TEXT("Expected an operator"), TokenParseRange.GetData());
 					}
-					if (*LastOp == EOperatorToken::ParenOpen && Op == EOperatorToken::ParenClose)
+					if (*LastOp == EOperatorToken::ParenOpen && Op == EOperatorToken::ParenClose &&
+					    (FunctionStack.IsEmpty() || FunctionStack.Top().OpeningOpStackSize != OperatorStack.Num()))
 					{
 						return ParseError(TEXT("Empty Parentheses"), TokenParseRange.GetData());
 					}
@@ -550,6 +662,21 @@ TVariant<FExpressionObject, FParseError> FEngine::Parse(
 			else if (Op == EOperatorToken::ParenOpen)
 			{
 				OperatorStack.Push(Op);
+
+				// If the top element on the expression is a constant, check if it is a function and if so, pop the constant
+				// out of the expression. 
+				if (!Expression.IsEmpty() && Expression.Top().IsType<FName>())
+				{
+					FName ConstantName = Expression.Top().Get<FName>();
+					const int32 *FunctionIndexPtr = FunctionNameIndex.Find(ConstantName);
+					if (FunctionIndexPtr == nullptr)
+					{
+						return ParseError(FString::Printf(TEXT("Unknown function '%s'"), *ConstantName.ToString()), LastTokenParseRange.GetData());
+					}
+
+					Expression.Pop();
+					FunctionStack.Push({*FunctionIndexPtr, 0, OperatorStack.Num(), Expression.Num(), LastTokenParseRange.GetData()});
+				}
 			}
 			else if (Op == EOperatorToken::ParenClose)
 			{
@@ -561,9 +688,33 @@ TVariant<FExpressionObject, FParseError> FEngine::Parse(
 					}
 					PushOperatorTokenToExpression(OperatorStack.Pop(false));
 				}
+
+				if (!FunctionStack.IsEmpty() && FunctionStack.Top().OpeningOpStackSize == OperatorStack.Num())
+				{
+					// Compute the number of arguments.
+					const FFunctionCallInfo CallInfo = FunctionStack.Pop();
+					const FFunctionInfo& FunctionInfo = Functions[CallInfo.FunctionIndex];
+					const int32 ArgumentCount = CallInfo.CountedCommas + (CallInfo.ExpressionSize != Expression.Num());
+
+					if (ArgumentCount != FunctionInfo.ArgumentCount)
+					{
+						return ParseError(FString::Printf(TEXT("Invalid argument count. Expected %d, got %d"), FunctionInfo.ArgumentCount, ArgumentCount), CallInfo.TokenStart); 
+					}
+
+					Expression.Push(FExpressionObject::OpElement(TInPlaceType<FExpressionObject::FFunctionRef>(), CallInfo.FunctionIndex));
+				}
+				
 				// Remove the now unneeded open parentheses.
 				OperatorStack.Pop(false);
 			}
+			else if (Op == EOperatorToken::Comma)
+			{
+				if (FunctionStack.IsEmpty() || FunctionStack.Top().OpeningOpStackSize != OperatorStack.Num())
+				{
+					return ParseError(TEXT("Unexpected comma"), TokenParseRange.GetData()); 
+				}
+				FunctionStack.Top().CountedCommas++;
+			}
 			else if (GetOperatorTokenInfo(Op).Associativity == EAssociativity::Left)
 			{
 				// Peel off all operators that have the same precedence or higher.
@@ -585,6 +736,7 @@ TVariant<FExpressionObject, FParseError> FEngine::Parse(
 		}
 
 		LastToken = Token;
+		LastTokenParseRange = ParseRange;
 		ParseRange = ParseRange.TrimStart();
 	} while(!ParseRange.IsEmpty());
 
@@ -619,7 +771,16 @@ float FEngine::Execute(
 	
 	for (const FExpressionObject::OpElement& Token: InExpressionObject.Expression)
 	{
-		if (const FExpressionObject::EOperator* Operator = Token.TryGet<FExpressionObject::EOperator>())
+		if (const float* Value = Token.TryGet<float>())
+		{
+			ValueStack.Push(*Value);
+		}
+		else if (const FName* Constant = Token.TryGet<FName>())
+		{
+			const float *ConstantValue = Constants.Find(*Constant);
+			ValueStack.Push(ConstantValue ? *ConstantValue : 0.0f);
+		}
+		else if (const FExpressionObject::EOperator* Operator = Token.TryGet<FExpressionObject::EOperator>())
 		{
 			switch(*Operator)
 			{
@@ -699,14 +860,18 @@ float FEngine::Execute(
 				break;  
 			}
 		}
-		else if (const FName* Constant = Token.TryGet<FName>())
+		else if (const FExpressionObject::FFunctionRef* FuncRef = Token.TryGet<FExpressionObject::FFunctionRef>())
 		{
-			const float *Value = Constants.Find(*Constant);
-			ValueStack.Push(Value ? *Value : 0.0f);
-		}
-		else if (const float* Value = Token.TryGet<float>())
-		{
-			ValueStack.Push(*Value);
+			const FFunctionInfo& FunctionInfo = Functions[FuncRef->Index];
+			check(FunctionInfo.ArgumentCount <= ValueStack.Num());
+
+			TArrayView<float> ValueView(ValueStack.GetData() + ValueStack.Num() - FunctionInfo.ArgumentCount, FunctionInfo.ArgumentCount);
+			const float FuncValue = FunctionInfo.FunctionPtr(ValueView);
+			for (int32 Index = 0; Index < FunctionInfo.ArgumentCount; Index++)
+			{
+				ValueStack.Pop(false);
+			}
+			ValueStack.Push(FuncValue);
 		}
 	}
 
diff --git a/Engine/Plugins/Experimental/Animation/CurveExpression/Source/Runtime/Public/ExpressionEvaluator.h b/Engine/Plugins/Experimental/Animation/CurveExpression/Source/Runtime/Public/ExpressionEvaluator.h
--- a/Engine/Plugins/Experimental/Animation/CurveExpression/Source/Runtime/Public/ExpressionEvaluator.h
+++ b/Engine/Plugins/Experimental/Animation/CurveExpression/Source/Runtime/Public/ExpressionEvaluator.h
@@ -25,8 +25,9 @@ struct FParseError
 
 struct FExpressionObject
 {
+	FString ToString() const;
 private:
-	enum class EOperator
+	enum class EOperator : int32
 	{
 		Negate,				// Negation operator.
 		Add,				// Add last two values on stack and add the result to the stack.
@@ -37,9 +38,15 @@ private:
 		Power,				// Raise to power
 		FloorDivide,    	// Divide and round the result down
 	};
+
+	struct FFunctionRef
+	{
+		FFunctionRef(int32 InIndex) : Index(InIndex) {}
+		int32 Index;
+	};
 	
 	friend class FEngine;
-	using OpElement = TVariant<EOperator, /* Constant Ref */ FName, /* Value */ float>;
+	using OpElement = TVariant<EOperator, FName /* Constant */, FFunctionRef, /* Value */ float>;
 	TArray<OpElement, TInlineAllocator<8>> Expression;
 };
 
@@ -141,6 +148,7 @@ private:
 		FloorDivide,    	// '//'
 		ParenOpen,			// '('
 		ParenClose,			// ')'
+		Comma,				// ','
 	};
 	using FToken = TVariant<EOperatorToken, /* Identifier */ FName, /* Value */ float>;
 	
@@ -184,10 +192,21 @@ private:
 	
 	/** The list of constants, stored using the FName key since the expression object refers
 	 *  to it directly */
-	// FUTURE: Add support for variables by allowing TFunction{Ref} objects along with direct
+	// TBD: Add support for variables by allowing TFunction{Ref} objects along with direct
 	// constants.
 	TMap<FName, float> Constants;
 
+	/** For now we just have built-in functions */
+	friend struct FInitializeBuiltinFunctions;
+	using FunctionType = TFunction<float(TArrayView<const float>)>;
+	struct FFunctionInfo
+	{
+		int32 ArgumentCount;
+		FunctionType FunctionPtr;
+	};
+	static TMap<FName, int32> FunctionNameIndex;
+	static TArray<FFunctionInfo> Functions;
+
 	/** The parse flags for this engine. Only set once */
 	EParseFlags ParseFlags = EParseFlags::Default;
 };
