Curve Expressions: Fix menu categoires

#rnx

diff --git a/Engine/Plugins/Experimental/Animation/CurveExpression/Source/Editor/Private/AnimGraphNode_RemapCurvesFromMesh.cpp b/Engine/Plugins/Experimental/Animation/CurveExpression/Source/Editor/Private/AnimGraphNode_RemapCurvesFromMesh.cpp
--- a/Engine/Plugins/Experimental/Animation/CurveExpression/Source/Editor/Private/AnimGraphNode_RemapCurvesFromMesh.cpp
+++ b/Engine/Plugins/Experimental/Animation/CurveExpression/Source/Editor/Private/AnimGraphNode_RemapCurvesFromMesh.cpp
@@ -35,6 +35,13 @@ FText UAnimGraphNode_RemapCurvesFromMesh::GetTooltipText() const
 	return LOCTEXT("Tooltip", "The Remap Curves From Mesh node copies curves from another component to this. Can be used to map any curve to any and perform mathematical operations on them.");
 }
 
+
+FText UAnimGraphNode_RemapCurvesFromMesh::GetMenuCategory() const
+{
+	return LOCTEXT("NodeCategory", "Animation|Curve Expression");
+}
+
+
 void UAnimGraphNode_RemapCurvesFromMesh::GetOutputLinkAttributes(FNodeAttributeArray& OutAttributes) const
 {
 	OutAttributes.Add(UE::Anim::FAttributes::Curves);
diff --git a/Engine/Plugins/Experimental/Animation/CurveExpression/Source/Editor/Private/K2Node_MakeCurveExpressionMap.cpp b/Engine/Plugins/Experimental/Animation/CurveExpression/Source/Editor/Private/K2Node_MakeCurveExpressionMap.cpp
--- a/Engine/Plugins/Experimental/Animation/CurveExpression/Source/Editor/Private/K2Node_MakeCurveExpressionMap.cpp
+++ b/Engine/Plugins/Experimental/Animation/CurveExpression/Source/Editor/Private/K2Node_MakeCurveExpressionMap.cpp
@@ -239,7 +239,7 @@ FNodeHandlingFunctor* UK2Node_MakeCurveExpressionMap::CreateNodeHandler(FKismetC
 
 FText UK2Node_MakeCurveExpressionMap::GetMenuCategory() const
 {
-	return LOCTEXT("NodeCategory", "Curve Expression");
+	return LOCTEXT("NodeCategory", "Animation|Curve Expression");
 }
 
 #undef LOCTEXT_NAMESPACE
diff --git a/Engine/Plugins/Experimental/Animation/CurveExpression/Source/Editor/Public/AnimGraphNode_RemapCurvesFromMesh.h b/Engine/Plugins/Experimental/Animation/CurveExpression/Source/Editor/Public/AnimGraphNode_RemapCurvesFromMesh.h
--- a/Engine/Plugins/Experimental/Animation/CurveExpression/Source/Editor/Public/AnimGraphNode_RemapCurvesFromMesh.h
+++ b/Engine/Plugins/Experimental/Animation/CurveExpression/Source/Editor/Public/AnimGraphNode_RemapCurvesFromMesh.h
@@ -25,6 +25,7 @@ public:
 	// UEdGraphNode interface
 	virtual FText GetNodeTitle(ENodeTitleType::Type TitleType) const override;
 	virtual FText GetTooltipText() const override;
+	FText GetMenuCategory() const override;
 	// End of UEdGraphNode interface
 
 	// UAnimGraphNode_Base interface
