Curve Expression: Optionally zero out constant values that are not in the updated constant value map. Fixes issues with the 3L Jonesy face getting stuck when facial curves are no longer coming from a clip.

diff --git a/Engine/Plugins/Experimental/Animation/CurveExpression/Source/Runtime/Private/ExpressionEvaluator.cpp b/Engine/Plugins/Experimental/Animation/CurveExpression/Source/Runtime/Private/ExpressionEvaluator.cpp
--- a/Engine/Plugins/Experimental/Animation/CurveExpression/Source/Runtime/Private/ExpressionEvaluator.cpp
+++ b/Engine/Plugins/Experimental/Animation/CurveExpression/Source/Runtime/Private/ExpressionEvaluator.cpp
@@ -4,6 +4,8 @@
 
 #include <limits>
 
+PRAGMA_DISABLE_OPTIMIZATION
+
 namespace CurveExpression::Evaluator
 {
 
@@ -126,14 +128,19 @@ FEngine::FEngine(
 
 
 void FEngine::UpdateConstantValues(
-	const TMap<FName, float>& InConstants
+	const TMap<FName, float>& InConstants,
+	EUpdateConstantsMethod InUpdateMethod
 	)
 {
-	for (const TTuple<FName, float>& NewConstantItem: InConstants)
+	for (TTuple<FName, float>& OldConstantItem: Constants)
 	{
-		if (float* OldConstant = Constants.Find(NewConstantItem.Key))
+		if (const float* NewConstant = InConstants.Find(OldConstantItem.Key))
+		{
+			OldConstantItem.Value = *NewConstant;
+		}
+		else if (InUpdateMethod == EUpdateConstantsMethod::ZeroMissing)
 		{
-			*OldConstant = NewConstantItem.Value; 
+			OldConstantItem.Value = 0.0f;
 		}
 	}
 }
@@ -906,3 +913,5 @@ TOptional<FParseError> FEngine::Verify(
 }
 
 }
+
+PRAGMA_ENABLE_OPTIMIZATION
diff --git a/Engine/Plugins/Experimental/Animation/CurveExpression/Source/Runtime/Public/ExpressionEvaluator.h b/Engine/Plugins/Experimental/Animation/CurveExpression/Source/Runtime/Public/ExpressionEvaluator.h
--- a/Engine/Plugins/Experimental/Animation/CurveExpression/Source/Runtime/Public/ExpressionEvaluator.h
+++ b/Engine/Plugins/Experimental/Animation/CurveExpression/Source/Runtime/Public/ExpressionEvaluator.h
@@ -81,6 +81,13 @@ enum class EParseFlags
 ENUM_CLASS_FLAGS(EParseFlags)
 
 
+/** Dictate how constants are updated. */
+enum class EUpdateConstantsMethod
+{
+	IgnoreMissing,					/** Existing constant values that are not updated are left unchanged. */
+	ZeroMissing						/** Existing constant values that are not updated are zeroed out. */
+};
+
 class CURVEEXPRESSION_API FEngine
 {
 public:
@@ -102,7 +109,8 @@ public:
 
 	/** Update existing constant values. No new constants are added. */
 	void UpdateConstantValues(
-		const TMap<FName, float>& InConstants
+		const TMap<FName, float>& InConstants,
+		EUpdateConstantsMethod InUpdateMethod = EUpdateConstantsMethod::ZeroMissing
 		);
 
 	/** Return existing constants and their values. */
