Curve Expression: Fix the verify button to work properly. We now search the BP's hierarchy to find the node. The warning is also added to the node itself rather than just the log.

diff --git a/Engine/Plugins/Experimental/Animation/CurveExpression/Source/Editor/Private/AnimGraphNode_RemapCurvesFromMesh.cpp b/Engine/Plugins/Experimental/Animation/CurveExpression/Source/Editor/Private/AnimGraphNode_RemapCurvesFromMesh.cpp
--- a/Engine/Plugins/Experimental/Animation/CurveExpression/Source/Editor/Private/AnimGraphNode_RemapCurvesFromMesh.cpp
+++ b/Engine/Plugins/Experimental/Animation/CurveExpression/Source/Editor/Private/AnimGraphNode_RemapCurvesFromMesh.cpp
@@ -1,6 +1,10 @@
 // Copyright Epic Games, Inc. All Rights Reserved.
 
 #include "AnimGraphNode_RemapCurvesFromMesh.h"
+
+#include "CurveExpressionModule.h"
+#include "SGraphNode.h"
+
 #include "Animation/AnimAttributes.h"
 #include "Animation/AnimInstance.h"
 #include "Components/SkeletalMeshComponent.h"
@@ -12,15 +16,37 @@
 
 bool UAnimGraphNode_RemapCurvesFromMesh::CanVerifyExpressions() const
 {
-	return GetDebuggedNode() != nullptr;
+	if (const FAnimNode_RemapCurvesFromMesh* DebuggedNode = GetDebuggedNode())
+	{
+		return DebuggedNode->CanVerifyExpressions();
+	}
+	return false;
 }
 
 
 void UAnimGraphNode_RemapCurvesFromMesh::VerifyExpressions()
 {
-	if (FAnimNode_RemapCurvesFromMesh* DebuggedNode = GetDebuggedNode())
+	if (const FAnimNode_RemapCurvesFromMesh* DebuggedNode = GetDebuggedNode())
 	{
-		DebuggedNode->VerifyExpressions();
+		TArray<FString> Results;
+		DebuggedNode->VerifyExpressions([&Results](const FString& InMessage){
+			Results.Add(InMessage);
+		});
+
+		// Don't override the current message, unless it's just an info message or 
+		if (!Results.IsEmpty() && (!bHasCompilerMessage || ErrorType > EMessageSeverity::Warning))
+		{
+			ErrorMsg = FString::Join(Results, TEXT("\n"));
+			ErrorType = EMessageSeverity::Warning;
+			bHasCompilerMessage = true;
+			
+			// This would be much nicer as a utility on UAnimGraphNode_Base
+			TSharedPtr<SGraphNode> Widget = DEPRECATED_NodeWidget.Pin();
+			if (Widget)
+			{
+				Widget->RefreshErrorInfo();
+			}
+		}
 	}
 }
 
diff --git a/Engine/Plugins/Experimental/Animation/CurveExpression/Source/Editor/Private/CurveExpressionDetailsCustomization.cpp b/Engine/Plugins/Experimental/Animation/CurveExpression/Source/Editor/Private/CurveExpressionDetailsCustomization.cpp
--- a/Engine/Plugins/Experimental/Animation/CurveExpression/Source/Editor/Private/CurveExpressionDetailsCustomization.cpp
+++ b/Engine/Plugins/Experimental/Animation/CurveExpression/Source/Editor/Private/CurveExpressionDetailsCustomization.cpp
@@ -68,7 +68,7 @@ void FAnimGraphNode_RemapCurvesFromMeshCustomization::CustomizeDetails(
 					{
 						if (UAnimGraphNode_RemapCurvesFromMesh* RemapNode = RemapNodeWeak.Get())
 						{
-							return true;
+							return RemapNode->CanVerifyExpressions();
 						}
 					}
 					return false;
diff --git a/Engine/Plugins/Experimental/Animation/CurveExpression/Source/Runtime/Private/AnimNode_RemapCurvesFromMesh.cpp b/Engine/Plugins/Experimental/Animation/CurveExpression/Source/Runtime/Private/AnimNode_RemapCurvesFromMesh.cpp
--- a/Engine/Plugins/Experimental/Animation/CurveExpression/Source/Runtime/Private/AnimNode_RemapCurvesFromMesh.cpp
+++ b/Engine/Plugins/Experimental/Animation/CurveExpression/Source/Runtime/Private/AnimNode_RemapCurvesFromMesh.cpp
@@ -29,16 +29,34 @@ void FAnimNode_RemapCurvesFromMesh::CacheBones_AnyThread(
 }
 	
 
-void FAnimNode_RemapCurvesFromMesh::VerifyExpressions()
+void FAnimNode_RemapCurvesFromMesh::VerifyExpressions(
+	TFunction<void(FString)> InReportingFunc
+	) const
 {
 	using namespace CurveExpression::Evaluator;
 
-	if (!ExpressionEngine.IsSet() || CurveExpressions.IsEmpty())
+	auto ReportAndLog = [InReportingFunc](FString InMessage)
 	{
+		if (InReportingFunc)
+		{
+			InReportingFunc(InMessage);
+		}
+		UE_LOG(LogCurveExpression, Warning, TEXT("%s"), *InMessage);
+	};
+	
+	if (!ExpressionEngine.IsSet())
+	{
+		return;
+	}
+	
+	if (CurveExpressions.IsEmpty())
+	{
+		ReportAndLog(TEXT("No curve expressions set."));
 		return;
 	}
 
-	FEngine VerificationEngine(
+	bool bFoundError = false;
+	const FEngine VerificationEngine(
 		ExpressionEngine->GetConstantValues(),
 		EParseFlags::ValidateConstants);
 
@@ -46,15 +64,29 @@ void FAnimNode_RemapCurvesFromMesh::VerifyExpressions()
 	{
 		if (!CurveNameToUIDMap.Contains(ExpressionPair.Key))
 		{
-			UE_LOG(LogCurveExpression, Warning, TEXT("Target curve '%s' does not exist."), *ExpressionPair.Key.ToString());
+			ReportAndLog(FString::Printf(TEXT("Target curve '%s' does not exist."), *ExpressionPair.Key.ToString()));
+			bFoundError = true;
 		}
 
 		TOptional<FParseError> Error = VerificationEngine.Verify(ExpressionPair.Value);
 		if (Error.IsSet())
 		{
-			UE_LOG(LogCurveExpression, Warning, TEXT("Expression error in '%s': %s"), *ExpressionPair.Value, *Error->Message);
+			ReportAndLog(FString::Printf(TEXT("Expression error in '%s': %s"), *ExpressionPair.Value, *Error->Message));
+			bFoundError = true;
 		}
 	}
+
+	if (!bFoundError)
+	{
+		UE_LOG(LogCurveExpression, Display, TEXT("Curve expressions verified ok."))
+	}
+}
+
+
+
+bool FAnimNode_RemapCurvesFromMesh::CanVerifyExpressions() const
+{
+	return ExpressionEngine.IsSet();
 }
 
 
diff --git a/Engine/Plugins/Experimental/Animation/CurveExpression/Source/Runtime/Public/AnimNode_RemapCurvesFromMesh.h b/Engine/Plugins/Experimental/Animation/CurveExpression/Source/Runtime/Public/AnimNode_RemapCurvesFromMesh.h
--- a/Engine/Plugins/Experimental/Animation/CurveExpression/Source/Runtime/Public/AnimNode_RemapCurvesFromMesh.h
+++ b/Engine/Plugins/Experimental/Animation/CurveExpression/Source/Runtime/Public/AnimNode_RemapCurvesFromMesh.h
@@ -31,8 +31,14 @@ struct CURVEEXPRESSION_API FAnimNode_RemapCurvesFromMesh :
 	UPROPERTY(EditAnywhere, BlueprintReadWrite, Category=Expressions, meta=(NeverAsPin))
 	bool bExpressionsImmutable = true;
 
-	// Call to verify the expressions. Report any errors back.
-	void VerifyExpressions();
+	// Call to verify the expressions. Write out any warnings to the InReportingFunc function.
+	void VerifyExpressions(
+		TFunction<void(FString)> InReportingFunc
+		) const;
+	
+	// Returns true if we're able to verify expressions at all. I.e. if the expression engine has not been initialized.
+	bool CanVerifyExpressions() const;
+
 	
 	// FAnimNode_Base interface
 	virtual void Initialize_AnyThread(const FAnimationInitializeContext& Context) override;
diff --git a/Engine/Source/Editor/AnimGraph/Private/AnimGraphNode_Base.cpp b/Engine/Source/Editor/AnimGraph/Private/AnimGraphNode_Base.cpp
--- a/Engine/Source/Editor/AnimGraph/Private/AnimGraphNode_Base.cpp
+++ b/Engine/Source/Editor/AnimGraph/Private/AnimGraphNode_Base.cpp
@@ -707,6 +707,7 @@ void UAnimGraphNode_Base::Draw(FPrimitiveDrawInterface* PDI, USkeletalMeshCompon
 	}
 }
 
+PRAGMA_DISABLE_OPTIMIZATION
 FAnimNode_Base* UAnimGraphNode_Base::FindDebugAnimNode(USkeletalMeshComponent* InPreviewSkelMeshComp) const
 {
 	FAnimNode_Base* DebugNode = nullptr;
@@ -724,18 +725,15 @@ FAnimNode_Base* UAnimGraphNode_Base::FindDebugAnimNode(USkeletalMeshComponent* I
 			{
 				return nullptr;
 			}
-			
-			FAnimBlueprintDebugData& DebugData = AnimBlueprintClass->GetAnimBlueprintDebugData();
-			int32* IndexPtr = DebugData.NodePropertyToIndexMap.Find(this);
-			if (!IndexPtr)
+
+			// Search for the node by GUID, since we can have multiple instantiations of this node, and therefore
+			// different pointers from ourselves even though they represent the same node in the graph.
+			const int32 AnimNodeIndex = AnimBlueprintClass->GetNodeIndexFromGuid(NodeGuid, EPropertySearchMode::Hierarchy);
+			if (AnimNodeIndex == INDEX_NONE)
 			{
 				return nullptr;
 			}
 			
-			int32 AnimNodeIndex = *IndexPtr;
-			// reverse node index temporarily because of a bug in NodeGuidToIndexMap
-			AnimNodeIndex = AnimBlueprintClass->GetAnimNodeProperties().Num() - AnimNodeIndex - 1;
-
 			return AnimBlueprintClass->GetAnimNodeProperties()[AnimNodeIndex]->ContainerPtrToValuePtr<FAnimNode_Base>(InAnimInstance);
 		};
 	
@@ -751,6 +749,7 @@ FAnimNode_Base* UAnimGraphNode_Base::FindDebugAnimNode(USkeletalMeshComponent* I
 
 	return DebugNode;
 }
+PRAGMA_ENABLE_OPTIMIZATION
 
 EAnimAssetHandlerType UAnimGraphNode_Base::SupportsAssetClass(const UClass* AssetClass) const
 {
