Fix a mismatch in array size for getting operator type from enum value and add a static assert

#rnx

diff --git a/Engine/Plugins/Experimental/Animation/CurveExpression/Source/Runtime/Private/ExpressionEvaluator.cpp b/Engine/Plugins/Experimental/Animation/CurveExpression/Source/Runtime/Private/ExpressionEvaluator.cpp
--- a/Engine/Plugins/Experimental/Animation/CurveExpression/Source/Runtime/Private/ExpressionEvaluator.cpp
+++ b/Engine/Plugins/Experimental/Animation/CurveExpression/Source/Runtime/Private/ExpressionEvaluator.cpp
@@ -556,7 +556,9 @@ TVariant<FExpressionObject, FParseError> FEngine::Parse(
 			{ EOperatorToken::FloorDivide,	3, EAssociativity::Left, FExpressionObject::EOperator::FloorDivide},
 			{ EOperatorToken::ParenOpen,	0, EAssociativity::None},
 			{ EOperatorToken::ParenClose,	0, EAssociativity::None},
+			{ EOperatorToken::Comma,		0, EAssociativity::None},
 		};
+		static_assert(sizeof(OperatorTokenInfo) / sizeof(FOperatorTokenInfo) == (int32)EOperatorToken::Max);
 		checkSlow(OperatorTokenInfo[static_cast<int32>(InOpToken)].OperatorToken == InOpToken);
 		return OperatorTokenInfo[static_cast<int32>(InOpToken)];
 	};
diff --git a/Engine/Plugins/Experimental/Animation/CurveExpression/Source/Runtime/Public/ExpressionEvaluator.h b/Engine/Plugins/Experimental/Animation/CurveExpression/Source/Runtime/Public/ExpressionEvaluator.h
--- a/Engine/Plugins/Experimental/Animation/CurveExpression/Source/Runtime/Public/ExpressionEvaluator.h
+++ b/Engine/Plugins/Experimental/Animation/CurveExpression/Source/Runtime/Public/ExpressionEvaluator.h
@@ -157,6 +157,8 @@ private:
 		ParenOpen,			// '('
 		ParenClose,			// ')'
 		Comma,				// ','
+
+		Max					// Must be last
 	};
 	using FToken = TVariant<EOperatorToken, /* Identifier */ FName, /* Value */ float>;
 	
