CurveExpressions: Keep evaluating curves even if the source mesh component doesn't have a valid anim instance at that point.

#rnx

diff --git a/Engine/Plugins/Experimental/Animation/CurveExpression/Source/Runtime/Private/AnimNode_RemapCurvesFromMesh.cpp b/Engine/Plugins/Experimental/Animation/CurveExpression/Source/Runtime/Private/AnimNode_RemapCurvesFromMesh.cpp
--- a/Engine/Plugins/Experimental/Animation/CurveExpression/Source/Runtime/Private/AnimNode_RemapCurvesFromMesh.cpp
+++ b/Engine/Plugins/Experimental/Animation/CurveExpression/Source/Runtime/Private/AnimNode_RemapCurvesFromMesh.cpp
@@ -10,6 +10,8 @@
 
 #include UE_INLINE_GENERATED_CPP_BY_NAME(AnimNode_RemapCurvesFromMesh)
 
+UE_DISABLE_OPTIMIZATION
+
 void FAnimNode_RemapCurvesFromMesh::Initialize_AnyThread(
 	const FAnimationInitializeContext& Context
 	)
@@ -207,10 +209,6 @@ void FAnimNode_RemapCurvesFromMesh::PreUpdate(
 				{
 					ExpressionEngine->UpdateConstantValues(SourceAnimInstance->GetAnimationCurveList(EAnimCurveType::AttributeCurve));
 				}
-				else
-				{
-					ExpressionEngine.Reset();
-				}
 			}
 		}
 		else
@@ -350,3 +348,4 @@ void FAnimNode_RemapCurvesFromMesh::RefreshMeshComponent(
 	}
 }
 
+UE_ENABLE_OPTIMIZATION
