CurveExpression: Remove stray optimization pragmas.

#rnx

diff --git a/Engine/Plugins/Experimental/Animation/CurveExpression/Source/Runtime/Private/AnimNode_RemapCurvesFromMesh.cpp b/Engine/Plugins/Experimental/Animation/CurveExpression/Source/Runtime/Private/AnimNode_RemapCurvesFromMesh.cpp
--- a/Engine/Plugins/Experimental/Animation/CurveExpression/Source/Runtime/Private/AnimNode_RemapCurvesFromMesh.cpp
+++ b/Engine/Plugins/Experimental/Animation/CurveExpression/Source/Runtime/Private/AnimNode_RemapCurvesFromMesh.cpp
@@ -10,8 +10,6 @@
 
 #include UE_INLINE_GENERATED_CPP_BY_NAME(AnimNode_RemapCurvesFromMesh)
 
-UE_DISABLE_OPTIMIZATION
-
 void FAnimNode_RemapCurvesFromMesh::Initialize_AnyThread(
 	const FAnimationInitializeContext& Context
 	)
@@ -347,5 +345,3 @@ void FAnimNode_RemapCurvesFromMesh::RefreshMeshComponent(
 		CurrentlyUsedSourceMeshComponent.Reset();
 	}
 }
-
-UE_ENABLE_OPTIMIZATION
