Curve Expression: Curve Expression List Data Asset

Added support for a curve expression data asset that can be swapped out at run-time. The asset contains compiled expressions so swapping in and out won't incur any runtime expression parsing costs.
Also added the ability to set expressions directly on the node.
Added a new node, "Remap Curves" which allows remapping curves within the same pose context.
Fixed the expression evaluator in preparation for new curve work.

diff --git a/Engine/Plugins/Experimental/Animation/CurveExpression/Source/Editor/CurveExpressionEditor.Build.cs b/Engine/Plugins/Experimental/Animation/CurveExpression/Source/Editor/CurveExpressionEditor.Build.cs
--- a/Engine/Plugins/Experimental/Animation/CurveExpression/Source/Editor/CurveExpressionEditor.Build.cs
+++ b/Engine/Plugins/Experimental/Animation/CurveExpression/Source/Editor/CurveExpressionEditor.Build.cs
@@ -1,25 +1,25 @@
 // Copyright Epic Games, Inc. All Rights Reserved.
 
+using EpicGames.Core;
 using UnrealBuildTool;
 
 public class CurveExpressionEditor : ModuleRules
 {
 	public CurveExpressionEditor(ReadOnlyTargetRules Target) : base(Target)
 	{
-		PrivateDependencyModuleNames.AddRange(
-			new string[]
-			{
-				"AnimGraph",
-				"BlueprintGraph",
-				"Core",
-				"CoreUObject",
-				"CurveExpression",
-				"Engine",
-				"KismetCompiler",
-				"Slate",
-				"SlateCore",
-				"UnrealEd",
-			}
+		PrivateDependencyModuleNames.AddAll(
+			"AnimGraph",
+			"AssetDefinition",
+			"BlueprintGraph",
+			"Core",
+			"CoreUObject",
+			"CurveExpression",
+			"Engine", 
+			"Kismet",
+			"KismetCompiler",
+			"Slate",
+			"SlateCore",
+			"UnrealEd"
 		);
 	}
 }
diff --git a/Engine/Plugins/Experimental/Animation/CurveExpression/Source/Editor/Private/AnimGraphNode_RemapCurves.cpp b/Engine/Plugins/Experimental/Animation/CurveExpression/Source/Editor/Private/AnimGraphNode_RemapCurves.cpp
new file mode 100644
--- /dev/null
+++ b/Engine/Plugins/Experimental/Animation/CurveExpression/Source/Editor/Private/AnimGraphNode_RemapCurves.cpp
@@ -0,0 +1,156 @@
+﻿// Copyright Epic Games, Inc. All Rights Reserved.
+
+#include "AnimGraphNode_RemapCurves.h"
+
+#include "SGraphNode.h"
+#include "Animation/AnimAttributes.h"
+#include "Animation/AnimInstance.h"
+#include "Components/SkeletalMeshComponent.h"
+#include "Kismet2/BlueprintEditorUtils.h"
+
+
+#define LOCTEXT_NAMESPACE "AnimGraphNode_RemapCurves"
+
+bool UAnimGraphNode_RemapCurves::CanVerifyExpressions() const
+{
+	if (const FAnimNode_RemapCurves* DebuggedNode = GetDebuggedNode())
+	{
+		return DebuggedNode->CanVerifyExpressions();
+	}
+	return false;
+}
+
+
+void UAnimGraphNode_RemapCurves::VerifyExpressions()
+{
+	if (const FAnimNode_RemapCurves* DebuggedNode = GetDebuggedNode())
+	{
+		const USkeletalMeshComponent* Component = GetDebuggedComponent();
+		
+		TArray<FString> Results;
+		DebuggedNode->VerifyExpressions(Component, Component, [&Results](const FString& InMessage){
+			Results.Add(InMessage);
+		});
+
+		// Don't override the current message, unless it's just an info message or 
+		if (!Results.IsEmpty() && (!bHasCompilerMessage || ErrorType > EMessageSeverity::Warning))
+		{
+			ErrorMsg = FString::Join(Results, TEXT("\n"));
+			ErrorType = EMessageSeverity::Warning;
+			bHasCompilerMessage = true;
+			
+			// This would be much nicer as a utility on UAnimGraphNode_Base
+			TSharedPtr<SGraphNode> Widget = DEPRECATED_NodeWidget.Pin();
+			if (Widget)
+			{
+				Widget->RefreshErrorInfo();
+			}
+		}
+	}
+}
+
+
+FText UAnimGraphNode_RemapCurves::GetNodeTitle(ENodeTitleType::Type TitleType) const
+{
+	return LOCTEXT("NodeTitle", "Remap Curves");
+}
+
+FText UAnimGraphNode_RemapCurves::GetTooltipText() const
+{
+	return LOCTEXT("Tooltip", "The Remap Curves node allows applying mathematical expressions on the pose's curve.");
+}
+
+
+FText UAnimGraphNode_RemapCurves::GetMenuCategory() const
+{
+	return LOCTEXT("NodeCategory", "Animation|Curve Expression");
+}
+
+
+void UAnimGraphNode_RemapCurves::GetOutputLinkAttributes(FNodeAttributeArray& OutAttributes) const
+{
+	OutAttributes.Add(UE::Anim::FAttributes::Curves);
+}
+
+
+void UAnimGraphNode_RemapCurves::CustomizePinData(UEdGraphPin* Pin, FName SourcePropertyName, int32 ArrayIndex) const
+{
+	Super::CustomizePinData(Pin, SourcePropertyName, ArrayIndex);
+
+	if (Pin->PinName == GET_MEMBER_NAME_STRING_CHECKED(FAnimNode_RemapCurvesBase, CurveExpressionsDataAsset))
+	{
+		Pin->bHidden = (Node.ExpressionSource != ERemapCurvesExpressionSource::DataAsset);
+	}
+	if (Pin->PinName == GET_MEMBER_NAME_STRING_CHECKED(FAnimNode_RemapCurvesBase, CurveExpressions))
+	{
+		Pin->bHidden = (Node.ExpressionSource != ERemapCurvesExpressionSource::ExpressionMap);
+	}
+}
+
+
+void UAnimGraphNode_RemapCurves::PostEditChangeProperty(FPropertyChangedEvent& PropertyChangedEvent)
+{
+	Super::PostEditChangeProperty(PropertyChangedEvent);
+
+	const FName PropertyName = PropertyChangedEvent.GetPropertyName();
+	if (PropertyName == GET_MEMBER_NAME_CHECKED(FCurveExpressionList, AssignmentExpressions) || 
+	 	PropertyName == GET_MEMBER_NAME_CHECKED(FAnimNode_RemapCurves, ExpressionSource))
+	{
+		Node.ParseAndCacheExpressions();
+	}
+	if (PropertyName == GET_MEMBER_NAME_CHECKED(FAnimNode_RemapCurvesBase, ExpressionSource))
+	{
+		{
+			FScopedTransaction Transaction(LOCTEXT("ChangeExpressionSource", "Change Expression Source"));
+			Modify();
+
+			for (UEdGraphPin* Pin: Pins)
+			{
+				if (Pin->PinName == GET_MEMBER_NAME_STRING_CHECKED(FAnimNode_RemapCurvesBase, CurveExpressionsDataAsset) &&
+					Node.ExpressionSource != ERemapCurvesExpressionSource::DataAsset)
+				{
+					Pin->BreakAllPinLinks();
+				}
+				if (Pin->PinName == GET_MEMBER_NAME_STRING_CHECKED(FAnimNode_RemapCurvesBase, CurveExpressions) &&
+					Node.ExpressionSource != ERemapCurvesExpressionSource::ExpressionMap)
+				{
+					Pin->BreakAllPinLinks();
+				}
+			}
+		}
+		ReconstructNode();
+		FBlueprintEditorUtils::MarkBlueprintAsStructurallyModified(GetBlueprint());
+	}
+}
+
+
+USkeletalMeshComponent* UAnimGraphNode_RemapCurves::GetDebuggedComponent() const
+{
+	if (const UObject* ObjectBeingDebugged = GetAnimBlueprint()->GetObjectBeingDebugged())
+	{
+		if (const UAnimInstance* InstanceBeingDebugged = Cast<const UAnimInstance>(ObjectBeingDebugged))
+		{
+			USkeletalMeshComponent* Component = InstanceBeingDebugged->GetSkelMeshComponent();
+			if (Component != nullptr && Component->GetAnimInstance() != nullptr)
+			{
+				return Component;
+			}
+		}
+	}
+	return nullptr;
+}
+
+
+FAnimNode_RemapCurves* UAnimGraphNode_RemapCurves::GetDebuggedNode() const
+{
+	if (USkeletalMeshComponent* Component = GetDebuggedComponent(); Component != nullptr)
+	{
+		return static_cast<FAnimNode_RemapCurves*>(FindDebugAnimNode(Component));
+	}
+
+	return nullptr;
+}
+
+
+#undef LOCTEXT_NAMESPACE
+ 
\ No newline at end of file
diff --git a/Engine/Plugins/Experimental/Animation/CurveExpression/Source/Editor/Private/AnimGraphNode_RemapCurvesFromMesh.cpp b/Engine/Plugins/Experimental/Animation/CurveExpression/Source/Editor/Private/AnimGraphNode_RemapCurvesFromMesh.cpp
--- a/Engine/Plugins/Experimental/Animation/CurveExpression/Source/Editor/Private/AnimGraphNode_RemapCurvesFromMesh.cpp
+++ b/Engine/Plugins/Experimental/Animation/CurveExpression/Source/Editor/Private/AnimGraphNode_RemapCurvesFromMesh.cpp
@@ -7,12 +7,12 @@
 #include "Animation/AnimAttributes.h"
 #include "Animation/AnimInstance.h"
 #include "Components/SkeletalMeshComponent.h"
+#include "Kismet2/BlueprintEditorUtils.h"
 
 #include UE_INLINE_GENERATED_CPP_BY_NAME(AnimGraphNode_RemapCurvesFromMesh)
 
 #define LOCTEXT_NAMESPACE "AnimGraphNode_RemapCurvesFromMesh"
 
-
 bool UAnimGraphNode_RemapCurvesFromMesh::CanVerifyExpressions() const
 {
 	if (const FAnimNode_RemapCurvesFromMesh* DebuggedNode = GetDebuggedNode())
@@ -27,8 +27,19 @@ void UAnimGraphNode_RemapCurvesFromMesh::VerifyExpressions()
 {
 	if (const FAnimNode_RemapCurvesFromMesh* DebuggedNode = GetDebuggedNode())
 	{
+		const USkeletalMeshComponent* TargetComponent = GetDebuggedComponent();
+		const USkeletalMeshComponent* SourceComponent = nullptr;
+		if (Node.SourceMeshComponent.IsValid())
+		{
+			SourceComponent = Node.SourceMeshComponent.Get();
+		}
+		else if (Node.bUseAttachedParent && TargetComponent)
+		{
+			SourceComponent = Cast<USkeletalMeshComponent>(TargetComponent->GetAttachParent());
+		}
+		
 		TArray<FString> Results;
-		DebuggedNode->VerifyExpressions([&Results](const FString& InMessage){
+		DebuggedNode->VerifyExpressions(TargetComponent, SourceComponent, [&Results](const FString& InMessage){
 			Results.Add(InMessage);
 		});
 
@@ -73,7 +84,57 @@ void UAnimGraphNode_RemapCurvesFromMesh::GetOutputLinkAttributes(FNodeAttributeA
 }
 
 
-FAnimNode_RemapCurvesFromMesh* UAnimGraphNode_RemapCurvesFromMesh::GetDebuggedNode() const
+void UAnimGraphNode_RemapCurvesFromMesh::CustomizePinData(UEdGraphPin* Pin, FName SourcePropertyName, int32 ArrayIndex) const
+{
+	Super::CustomizePinData(Pin, SourcePropertyName, ArrayIndex);
+
+	if (Pin->PinName == GET_MEMBER_NAME_STRING_CHECKED(FAnimNode_RemapCurvesBase, CurveExpressionsDataAsset))
+	{
+		Pin->bHidden = (Node.ExpressionSource != ERemapCurvesExpressionSource::DataAsset);
+	}
+	if (Pin->PinName == GET_MEMBER_NAME_STRING_CHECKED(FAnimNode_RemapCurvesBase, CurveExpressions))
+	{
+		Pin->bHidden = (Node.ExpressionSource != ERemapCurvesExpressionSource::ExpressionMap);
+	}
+}
+
+
+void UAnimGraphNode_RemapCurvesFromMesh::PostEditChangeProperty(FPropertyChangedEvent& PropertyChangedEvent)
+{
+	Super::PostEditChangeProperty(PropertyChangedEvent);
+
+	const FName PropertyName = PropertyChangedEvent.GetPropertyName();
+	if (PropertyName == GET_MEMBER_NAME_CHECKED(FCurveExpressionList, AssignmentExpressions) || 
+	 	PropertyName == GET_MEMBER_NAME_CHECKED(FAnimNode_RemapCurvesFromMesh, ExpressionSource))
+	{
+		Node.ParseAndCacheExpressions();
+	}
+	if (PropertyName == GET_MEMBER_NAME_CHECKED(FAnimNode_RemapCurvesBase, ExpressionSource))
+	{
+		{
+			FScopedTransaction Transaction(LOCTEXT("ChangeExpressionSource", "Change Expression Source"));
+			Modify();
+
+			for (UEdGraphPin* Pin: Pins)
+			{
+				if (Pin->PinName == GET_MEMBER_NAME_STRING_CHECKED(FAnimNode_RemapCurvesBase, CurveExpressionsDataAsset) &&
+					Node.ExpressionSource != ERemapCurvesExpressionSource::DataAsset)
+				{
+					Pin->BreakAllPinLinks();
+				}
+				if (Pin->PinName == GET_MEMBER_NAME_STRING_CHECKED(FAnimNode_RemapCurvesBase, CurveExpressions) &&
+					Node.ExpressionSource != ERemapCurvesExpressionSource::ExpressionMap)
+				{
+					Pin->BreakAllPinLinks();
+				}
+			}
+		}
+		ReconstructNode();
+		FBlueprintEditorUtils::MarkBlueprintAsStructurallyModified(GetBlueprint());
+	}
+}
+
+USkeletalMeshComponent* UAnimGraphNode_RemapCurvesFromMesh::GetDebuggedComponent() const
 {
 	if (const UObject* ObjectBeingDebugged = GetAnimBlueprint()->GetObjectBeingDebugged())
 	{
@@ -82,10 +143,20 @@ FAnimNode_RemapCurvesFromMesh* UAnimGraphNode_RemapCurvesFromMesh::GetDebuggedNo
 			USkeletalMeshComponent* Component = InstanceBeingDebugged->GetSkelMeshComponent();
 			if (Component != nullptr && Component->GetAnimInstance() != nullptr)
 			{
-				return static_cast<FAnimNode_RemapCurvesFromMesh*>(FindDebugAnimNode(Component));
+				return Component;
 			}
 		}
 	}
+	return nullptr;
+}
+
+
+FAnimNode_RemapCurvesFromMesh* UAnimGraphNode_RemapCurvesFromMesh::GetDebuggedNode() const
+{
+	if (USkeletalMeshComponent* Component = GetDebuggedComponent(); Component != nullptr)
+	{
+		return static_cast<FAnimNode_RemapCurvesFromMesh*>(FindDebugAnimNode(Component));
+	}
 
 	return nullptr;
 }
diff --git a/Engine/Plugins/Experimental/Animation/CurveExpression/Source/Editor/Private/AssetDefinition_CurveExpressionsDataAsset.cpp b/Engine/Plugins/Experimental/Animation/CurveExpression/Source/Editor/Private/AssetDefinition_CurveExpressionsDataAsset.cpp
new file mode 100644
--- /dev/null
+++ b/Engine/Plugins/Experimental/Animation/CurveExpression/Source/Editor/Private/AssetDefinition_CurveExpressionsDataAsset.cpp
@@ -0,0 +1,17 @@
+﻿// Copyright Epic Games, Inc. All Rights Reserved.
+
+#include "AssetDefinition_CurveExpressionsDataAsset.h"
+
+#include "SDetailsDiff.h"
+
+
+EAssetCommandResult UAssetDefinition_CurveExpressionsDataAsset::PerformAssetDiff(const FAssetDiffArgs& DiffArgs) const
+{
+	if (DiffArgs.OldAsset == nullptr && DiffArgs.NewAsset == nullptr)
+	{
+		return EAssetCommandResult::Unhandled;
+	}
+	
+	SDetailsDiff::CreateDiffWindow(DiffArgs.OldAsset, DiffArgs.NewAsset, DiffArgs.OldRevision, DiffArgs.NewRevision, UCurveExpressionsDataAsset::StaticClass());
+	return EAssetCommandResult::Handled;
+}
diff --git a/Engine/Plugins/Experimental/Animation/CurveExpression/Source/Editor/Private/AssetDefinition_CurveExpressionsDataAsset.h b/Engine/Plugins/Experimental/Animation/CurveExpression/Source/Editor/Private/AssetDefinition_CurveExpressionsDataAsset.h
new file mode 100644
--- /dev/null
+++ b/Engine/Plugins/Experimental/Animation/CurveExpression/Source/Editor/Private/AssetDefinition_CurveExpressionsDataAsset.h
@@ -0,0 +1,27 @@
+﻿// Copyright Epic Games, Inc. All Rights Reserved.
+
+#pragma once
+
+#include "CurveExpressionsDataAsset.h"
+
+#include "AssetDefinitionDefault.h"
+
+#include "AssetDefinition_CurveExpressionsDataAsset.generated.h"
+
+UCLASS()
+class UAssetDefinition_CurveExpressionsDataAsset :
+	public UAssetDefinitionDefault
+{
+	GENERATED_BODY()
+
+public:
+	virtual FText GetAssetDisplayName() const override { return NSLOCTEXT("CurveExpressionEditor", "CurveExpressionDataAsset", "Curve Expression Data Asset"); }
+	virtual FLinearColor GetAssetColor() const override { return FLinearColor(FColor(201, 29, 85)); }
+	virtual TSoftClassPtr<UObject> GetAssetClass() const override { return UCurveExpressionsDataAsset::StaticClass(); }
+	virtual TConstArrayView<FAssetCategoryPath> GetAssetCategories() const override
+	{
+		static const auto Categories = { EAssetCategoryPaths::Animation / NSLOCTEXT("CurveExpressionEditor", "CurveExpressionAssetMenu", "Curve Expressions") };
+		return Categories;
+	}
+	virtual EAssetCommandResult PerformAssetDiff(const FAssetDiffArgs& DiffArgs) const override;
+};
diff --git a/Engine/Plugins/Experimental/Animation/CurveExpression/Source/Editor/Private/CurveExpressionDetailsCustomization.cpp b/Engine/Plugins/Experimental/Animation/CurveExpression/Source/Editor/Private/CurveExpressionDetailsCustomization.cpp
--- a/Engine/Plugins/Experimental/Animation/CurveExpression/Source/Editor/Private/CurveExpressionDetailsCustomization.cpp
+++ b/Engine/Plugins/Experimental/Animation/CurveExpression/Source/Editor/Private/CurveExpressionDetailsCustomization.cpp
@@ -9,33 +9,32 @@
 #include "DetailLayoutBuilder.h"
 #include "DetailWidgetRow.h"
 #include "Widgets/Input/SButton.h"
+#include "Widgets/Input/SMultiLineEditableTextBox.h"
 #include "Widgets/Layout/SGridPanel.h"
 #include "Widgets/Layout/SScrollBox.h"
-#include "Widgets/Text/SMultiLineEditableText.h"
 
 
 #define LOCTEXT_NAMESPACE "CurveExpressionDetailsCustomization"
 
-TSharedRef<IDetailCustomization> FAnimGraphNode_RemapCurvesFromMeshCustomization::MakeInstance()
+TSharedRef<IDetailCustomization> FAnimGraphNode_RemapCurvesDebuggingCustomization::MakeInstance()
 {
-	return MakeShared<FAnimGraphNode_RemapCurvesFromMeshCustomization>();
+	return MakeShared<FAnimGraphNode_RemapCurvesDebuggingCustomization>();
 }
 
 
-void FAnimGraphNode_RemapCurvesFromMeshCustomization::CustomizeDetails(
+void FAnimGraphNode_RemapCurvesDebuggingCustomization::CustomizeDetails(
 	IDetailLayoutBuilder& InDetailBuilder
 	)
 {
 	TArray<TWeakObjectPtr<UObject>> ObjectsBeingCustomized;
 	InDetailBuilder.GetObjectsBeingCustomized(ObjectsBeingCustomized);
 
-	TArray<TWeakObjectPtr<UAnimGraphNode_RemapCurvesFromMesh>> RemapNodes;
+	TArray<TWeakObjectPtr<UObject>> RemapNodes;
 	for (TWeakObjectPtr<UObject> Object: ObjectsBeingCustomized)
 	{
-		TWeakObjectPtr<UAnimGraphNode_RemapCurvesFromMesh> RemapNode = Cast<UAnimGraphNode_RemapCurvesFromMesh>(Object.Get());
-		if (RemapNode.IsValid())
+		if (Cast<IRemapCurvesDebuggingProvider>(Object.Get()))
 		{
-			RemapNodes.Add(RemapNode);
+			RemapNodes.Add(Object);
 		}
 	}
 
@@ -53,9 +52,9 @@ void FAnimGraphNode_RemapCurvesFromMeshCustomization::CustomizeDetails(
 				.ToolTipText(LOCTEXT("VerifyExpressions_Tooltip", "Verify all the expressions as applied to the currently debugged object. If debugging is not active, no verification is done. Verification output is sent to the output log."))
 				.OnClicked_Lambda([this, RemapNodes]()
 				{
-					for (TWeakObjectPtr<UAnimGraphNode_RemapCurvesFromMesh> RemapNodeWeak: RemapNodes)
+					for (TWeakObjectPtr<UObject> Object: RemapNodes)
 					{
-						if (UAnimGraphNode_RemapCurvesFromMesh* RemapNode = RemapNodeWeak.Get())
+						if (IRemapCurvesDebuggingProvider* RemapNode = Cast<IRemapCurvesDebuggingProvider>(Object.Get()))
 						{
 							RemapNode->VerifyExpressions();
 						}
@@ -64,9 +63,9 @@ void FAnimGraphNode_RemapCurvesFromMeshCustomization::CustomizeDetails(
 				})
 				.IsEnabled_Lambda([this, RemapNodes]()
 				{
-					for (TWeakObjectPtr<UAnimGraphNode_RemapCurvesFromMesh> RemapNodeWeak: RemapNodes)
+					for (TWeakObjectPtr<UObject> Object: RemapNodes)
 					{
-						if (UAnimGraphNode_RemapCurvesFromMesh* RemapNode = RemapNodeWeak.Get())
+						if (const IRemapCurvesDebuggingProvider* RemapNode = Cast<IRemapCurvesDebuggingProvider>(Object.Get()))
 						{
 							return RemapNode->CanVerifyExpressions();
 						}
@@ -92,18 +91,24 @@ void FCurveExpressionListCustomization::CustomizeHeader(
 {
 	AssignmentExpressionsProperty = InPropertyHandle->GetChildHandle(GET_MEMBER_NAME_CHECKED(FCurveExpressionList, AssignmentExpressions));
 
+	const FScrollBarStyle& ScrollBarStyle = FCurveExpressionEditorStyle::Get().GetWidgetStyle<FScrollBarStyle>("TextEditor.ScrollBar"); 
+
 	HorizontalScrollbar =
 		SNew(SScrollBar)
+			.Style(&ScrollBarStyle)
 			.AlwaysShowScrollbar(true)
-			.Orientation(Orient_Horizontal);
+			.Orientation(Orient_Horizontal)
+			.Thickness(FVector2D(6.0f, 6.0f));
 
 	VerticalScrollbar =
 		SNew(SScrollBar)
+			.Style(&ScrollBarStyle)
 			.AlwaysShowScrollbar(true)
-			.Orientation(Orient_Vertical);
+			.Orientation(Orient_Vertical)
+			.Thickness(FVector2D(6.0f, 6.0f));
 
-	const FTextBlockStyle &TextStyle = FCurveExpressionEditorStyle::Get().GetWidgetStyle<FTextBlockStyle>("TextEditor.NormalText");
-	const FSlateFontInfo &Font = TextStyle.Font;
+	const FEditableTextBoxStyle &TextBoxStyle = FCurveExpressionEditorStyle::Get().GetWidgetStyle<FEditableTextBoxStyle>("TextEditor.EditableTextBox");
+	const FSlateFontInfo &Font = TextBoxStyle.TextStyle.Font;
 	
 	InHeaderRow
 	.WholeRowContent()
@@ -122,44 +127,24 @@ void FCurveExpressionListCustomization::CustomizeHeader(
 			.BorderImage(FCurveExpressionEditorStyle::Get().GetBrush("TextEditor.Border"))
 			.BorderBackgroundColor(FLinearColor::Black)
 			[
-				SNew(SGridPanel)
-				.FillColumn(0, 1.0f)
-				.FillRow(0, 1.0f)
-				+SGridPanel::Slot(0, 0)
-				[
-					SNew(SScrollBox)
-					.Orientation(Orient_Vertical)
-					.ExternalScrollbar(VerticalScrollbar)
-					+ SScrollBox::Slot()
-					[
-						SAssignNew(TextEditor, SMultiLineEditableText)
-						.Font(Font)
-						.TextStyle(&TextStyle)
-						.Text_Lambda([this]()
-						{
-							FString Text;
-							AssignmentExpressionsProperty->GetValue(Text);
-							return FText::FromString(Text);
-						})
-						.OnTextChanged_Lambda([this](const FText& InText)
-						{
-							AssignmentExpressionsProperty->SetValue(InText.ToString());
-						})
-						// By default, the Tab key gets routed to "next widget". We want to disable that behaviour.
-						.OnIsTypedCharValid_Lambda([](const TCHAR InChar) { return true; })
-						.AutoWrapText(false)
-						.HScrollBar(HorizontalScrollbar)
-						.VScrollBar(VerticalScrollbar)
-					]
-				]
-				+SGridPanel::Slot(1, 0)
-				[
-					VerticalScrollbar.ToSharedRef()
-				]
-				+SGridPanel::Slot(0, 1)
-				[
-					HorizontalScrollbar.ToSharedRef()
-				]
+				SAssignNew(TextEditor, SMultiLineEditableTextBox)
+				.Font(Font)
+				.Style(&FCoreStyle::Get().GetWidgetStyle<FEditableTextBoxStyle>("NormalEditableTextBox"))
+				.Text_Lambda([this]()
+				{
+					FString Text;
+					AssignmentExpressionsProperty->GetValue(Text);
+					return FText::FromString(Text);
+				})
+				.OnTextChanged_Lambda([this](const FText& InText)
+				{
+					AssignmentExpressionsProperty->SetValue(InText.ToString());
+				})
+				// By default, the Tab key gets routed to "next widget". We want to disable that behaviour.
+				.OnIsTypedCharValid_Lambda([](const TCHAR InChar) { return true; })
+				.AutoWrapText(false)
+				.HScrollBar(HorizontalScrollbar)
+				.VScrollBar(VerticalScrollbar)
 			]
 		]
 	];
diff --git a/Engine/Plugins/Experimental/Animation/CurveExpression/Source/Editor/Private/CurveExpressionDetailsCustomization.h b/Engine/Plugins/Experimental/Animation/CurveExpression/Source/Editor/Private/CurveExpressionDetailsCustomization.h
--- a/Engine/Plugins/Experimental/Animation/CurveExpression/Source/Editor/Private/CurveExpressionDetailsCustomization.h
+++ b/Engine/Plugins/Experimental/Animation/CurveExpression/Source/Editor/Private/CurveExpressionDetailsCustomization.h
@@ -6,11 +6,11 @@
 #include "IPropertyTypeCustomization.h"
 
 
-class SMultiLineEditableText;
+class SMultiLineEditableTextBox;
 class SScrollBar;
 
 
-class FAnimGraphNode_RemapCurvesFromMeshCustomization :
+class FAnimGraphNode_RemapCurvesDebuggingCustomization :
 	public IDetailCustomization
 {
 public:
@@ -44,6 +44,6 @@ private:
 	TSharedPtr<SScrollBar> HorizontalScrollbar;
 	TSharedPtr<SScrollBar> VerticalScrollbar;
 
-	TSharedPtr<SMultiLineEditableText> TextEditor;	
+	TSharedPtr<SMultiLineEditableTextBox> TextEditor;	
 };
 
diff --git a/Engine/Plugins/Experimental/Animation/CurveExpression/Source/Editor/Private/CurveExpressionEditorModule.cpp b/Engine/Plugins/Experimental/Animation/CurveExpression/Source/Editor/Private/CurveExpressionEditorModule.cpp
--- a/Engine/Plugins/Experimental/Animation/CurveExpression/Source/Editor/Private/CurveExpressionEditorModule.cpp
+++ b/Engine/Plugins/Experimental/Animation/CurveExpression/Source/Editor/Private/CurveExpressionEditorModule.cpp
@@ -1,5 +1,6 @@
 // Copyright Epic Games, Inc. All Rights Reserved.
 
+#include "AnimGraphNode_RemapCurves.h"
 #include "AnimGraphNode_RemapCurvesFromMesh.h"
 #include "CurveExpressionDetailsCustomization.h"
 #include "CurveExpressionEditorStyle.h"
@@ -40,9 +41,14 @@ public:
 			&FCurveExpressionListCustomization::MakeInstance
 			);
 			
+		RegisterClassCustomization(
+			UAnimGraphNode_RemapCurves::StaticClass()->GetFName(),
+			&FAnimGraphNode_RemapCurvesDebuggingCustomization::MakeInstance
+			);
+		
 		RegisterClassCustomization(
 			UAnimGraphNode_RemapCurvesFromMesh::StaticClass()->GetFName(),
-			&FAnimGraphNode_RemapCurvesFromMeshCustomization::MakeInstance
+			&FAnimGraphNode_RemapCurvesDebuggingCustomization::MakeInstance
 			);
 	}
 	virtual void ShutdownModule() override
diff --git a/Engine/Plugins/Experimental/Animation/CurveExpression/Source/Editor/Private/CurveExpressionEditorStyle.cpp b/Engine/Plugins/Experimental/Animation/CurveExpression/Source/Editor/Private/CurveExpressionEditorStyle.cpp
--- a/Engine/Plugins/Experimental/Animation/CurveExpression/Source/Editor/Private/CurveExpressionEditorStyle.cpp
+++ b/Engine/Plugins/Experimental/Animation/CurveExpression/Source/Editor/Private/CurveExpressionEditorStyle.cpp
@@ -54,19 +54,18 @@ FCurveExpressionEditorStyle::FCurveExpressionEditorStyle() :
 		const FTextBlockStyle ErrorText = FTextBlockStyle(NormalText)
 			.SetUnderlineBrush(IMAGE_BRUSH("Images/White", Icon8x8, FLinearColor::Red, ESlateBrushTileType::Both))
 			.SetColorAndOpacity(FLinearColor::Red);
-		// ;
 		
 		Set("TextEditor.NormalText", NormalText);
 		
 		Set("TextEditor.Border", new BOX_BRUSH("Images/TextEditorBorder", FMargin(4.0f/16.0f), FLinearColor(0.02f,0.02f,0.02f,1)));
 
-		const FEditableTextBoxStyle EditableTextBoxStyle = FEditableTextBoxStyle()
-			.SetTextStyle(NormalText)
-			.SetBackgroundImageNormal( FSlateNoResource() )
-			.SetBackgroundImageHovered( FSlateNoResource() )
-			.SetBackgroundImageFocused( FSlateNoResource() )
-			.SetBackgroundImageReadOnly( FSlateNoResource() );
+		const FEditableTextBoxStyle EditableTextBoxStyle = FEditableTextBoxStyle(FCoreStyle::Get().GetWidgetStyle<FEditableTextBoxStyle>("NormalEditableTextBox"))
+			.SetTextStyle(NormalText);
 		
 		Set("TextEditor.EditableTextBox", EditableTextBoxStyle);
+
+		const FScrollBarStyle ScrollBarStyle = FScrollBarStyle(FCoreStyle::Get().GetWidgetStyle<FScrollBarStyle>("ScrollBar"))
+			.SetThickness(6.0f);
+		Set("TextEditor.ScrollBar", ScrollBarStyle);
 	}
 }
diff --git a/Engine/Plugins/Experimental/Animation/CurveExpression/Source/Editor/Private/CurveExpressionsDataAssetFactory.cpp b/Engine/Plugins/Experimental/Animation/CurveExpression/Source/Editor/Private/CurveExpressionsDataAssetFactory.cpp
new file mode 100644
--- /dev/null
+++ b/Engine/Plugins/Experimental/Animation/CurveExpression/Source/Editor/Private/CurveExpressionsDataAssetFactory.cpp
@@ -0,0 +1,24 @@
+﻿// Copyright Epic Games, Inc. All Rights Reserved.
+
+#include "CurveExpressionsDataAssetFactory.h"
+
+#include "CurveExpressionsDataAsset.h"
+
+UCurveExpressionsDataAssetFactory::UCurveExpressionsDataAssetFactory()
+{
+	SupportedClass = UCurveExpressionsDataAsset::StaticClass();
+	bCreateNew = true;
+	bEditAfterNew = true;
+}
+
+UObject* UCurveExpressionsDataAssetFactory::FactoryCreateNew(
+	UClass* InClass, 
+	UObject* InParent, 
+	FName InName,
+	EObjectFlags InFlags, 
+	UObject* InContext,
+	FFeedbackContext* InWarn
+	)
+{
+	return NewObject<UCurveExpressionsDataAsset>(InParent, InClass, InName, InFlags);
+}
diff --git a/Engine/Plugins/Experimental/Animation/CurveExpression/Source/Editor/Private/CurveExpressionsDataAssetFactory.h b/Engine/Plugins/Experimental/Animation/CurveExpression/Source/Editor/Private/CurveExpressionsDataAssetFactory.h
new file mode 100644
--- /dev/null
+++ b/Engine/Plugins/Experimental/Animation/CurveExpression/Source/Editor/Private/CurveExpressionsDataAssetFactory.h
@@ -0,0 +1,19 @@
+﻿// Copyright Epic Games, Inc. All Rights Reserved.
+
+#pragma once
+
+#include "Factories/Factory.h"
+
+#include "CurveExpressionsDataAssetFactory.generated.h"
+
+UCLASS(MinimalAPI)
+class UCurveExpressionsDataAssetFactory :
+	public UFactory
+{
+	GENERATED_BODY()
+public:
+	UCurveExpressionsDataAssetFactory();
+
+	// UFactory overrides
+	virtual UObject* FactoryCreateNew(UClass* InClass, UObject* InParent, FName InName, EObjectFlags InFlags, UObject* InContext, FFeedbackContext* InWarn) override;
+};
diff --git a/Engine/Plugins/Experimental/Animation/CurveExpression/Source/Editor/Private/K2Node_MakeCurveExpressionMap.cpp b/Engine/Plugins/Experimental/Animation/CurveExpression/Source/Editor/Private/K2Node_MakeCurveExpressionMap.cpp
--- a/Engine/Plugins/Experimental/Animation/CurveExpression/Source/Editor/Private/K2Node_MakeCurveExpressionMap.cpp
+++ b/Engine/Plugins/Experimental/Animation/CurveExpression/Source/Editor/Private/K2Node_MakeCurveExpressionMap.cpp
@@ -196,7 +196,7 @@ void UK2Node_MakeCurveExpressionMap::ValidateNodeDuringCompilation(FCompilerResu
 				{
 					const int32 ExpressionColumn = SourceExpression.GetData() - InLine.GetData();
 					Args.Add(TEXT("Error"), FText::FromString(Error->Message));
-					Args.Add(TEXT("Column"), FText::AsNumber(Error->Column + ExpressionColumn + 1));
+					Args.Add(TEXT("Column"), FText::AsNumber(Error->Location.Start + ExpressionColumn + 1));
 				
 					MessageLog.Error(
 						*FText::Format(LOCTEXT("Warning_BadExpression", "@@ has an error on line {Line}, column {Column}. {Error}"), Args).ToString(), 
diff --git a/Engine/Plugins/Experimental/Animation/CurveExpression/Source/Editor/Public/AnimGraphNode_RemapCurves.h b/Engine/Plugins/Experimental/Animation/CurveExpression/Source/Editor/Public/AnimGraphNode_RemapCurves.h
new file mode 100644
--- /dev/null
+++ b/Engine/Plugins/Experimental/Animation/CurveExpression/Source/Editor/Public/AnimGraphNode_RemapCurves.h
@@ -0,0 +1,48 @@
+﻿// Copyright Epic Games, Inc. All Rights Reserved.
+
+#pragma once
+
+#include "AnimNode_RemapCurves.h"
+
+#include "AnimGraphNode_Base.h"
+#include "IRemapCurvesDebuggingProvider.h"
+
+#include "AnimGraphNode_RemapCurves.generated.h"
+
+
+namespace ENodeTitleType { enum Type : int; }
+
+
+UCLASS(MinimalAPI)
+class UAnimGraphNode_RemapCurves :
+	public UAnimGraphNode_Base,
+	public IRemapCurvesDebuggingProvider
+{
+	GENERATED_BODY()
+
+public:
+	UPROPERTY(EditAnywhere, BlueprintReadWrite, Category=Settings)
+	FAnimNode_RemapCurves Node;
+
+	bool CanVerifyExpressions() const override;
+	void VerifyExpressions() override;
+
+	// UEdGraphNode interface
+	FText GetNodeTitle(ENodeTitleType::Type TitleType) const override;
+	FText GetTooltipText() const override;
+	FText GetMenuCategory() const override;
+	// End of UEdGraphNode interface
+
+	// UAnimGraphNode_Base interface
+	bool UsingCopyPoseFromMesh() const override { return true; }
+	void GetOutputLinkAttributes(FNodeAttributeArray& OutAttributes) const override;
+	void CustomizePinData(UEdGraphPin* Pin, FName SourcePropertyName, int32 ArrayIndex) const override;
+	// End of UAnimGraphNode_Base interface
+	
+	// UObject override
+	void PostEditChangeProperty(FPropertyChangedEvent& PropertyChangedEvent) override; 	
+	
+private:
+	USkeletalMeshComponent* GetDebuggedComponent() const;
+	FAnimNode_RemapCurves* GetDebuggedNode() const;	
+};
diff --git a/Engine/Plugins/Experimental/Animation/CurveExpression/Source/Editor/Public/AnimGraphNode_RemapCurvesFromMesh.h b/Engine/Plugins/Experimental/Animation/CurveExpression/Source/Editor/Public/AnimGraphNode_RemapCurvesFromMesh.h
--- a/Engine/Plugins/Experimental/Animation/CurveExpression/Source/Editor/Public/AnimGraphNode_RemapCurvesFromMesh.h
+++ b/Engine/Plugins/Experimental/Animation/CurveExpression/Source/Editor/Public/AnimGraphNode_RemapCurvesFromMesh.h
@@ -5,6 +5,7 @@
 #include "AnimGraphNode_Base.h"
 
 #include "AnimNode_RemapCurvesFromMesh.h"
+#include "IRemapCurvesDebuggingProvider.h"
 
 #include "AnimGraphNode_RemapCurvesFromMesh.generated.h"
 
@@ -13,29 +14,34 @@ namespace ENodeTitleType { enum Type : int; }
 
 UCLASS(MinimalAPI)
 class UAnimGraphNode_RemapCurvesFromMesh :
-	public UAnimGraphNode_Base
+	public UAnimGraphNode_Base,
+	public IRemapCurvesDebuggingProvider
 {
 	GENERATED_BODY()
 
 public:
-	UPROPERTY(EditAnywhere, Category=Settings)
+	UPROPERTY(EditAnywhere, BlueprintReadWrite, Category=Settings)
 	FAnimNode_RemapCurvesFromMesh Node;
 
-	bool CanVerifyExpressions() const;
-	void VerifyExpressions();
+	bool CanVerifyExpressions() const override;
+	void VerifyExpressions() override;
 
 	// UEdGraphNode interface
-	virtual FText GetNodeTitle(ENodeTitleType::Type TitleType) const override;
-	virtual FText GetTooltipText() const override;
+	FText GetNodeTitle(ENodeTitleType::Type TitleType) const override;
+	FText GetTooltipText() const override;
 	FText GetMenuCategory() const override;
 	// End of UEdGraphNode interface
 
 	// UAnimGraphNode_Base interface
-	virtual bool UsingCopyPoseFromMesh() const override { return true; }
-	virtual void GetOutputLinkAttributes(FNodeAttributeArray& OutAttributes) const override;
+	bool UsingCopyPoseFromMesh() const override { return true; }
+	void GetOutputLinkAttributes(FNodeAttributeArray& OutAttributes) const override;
+	void CustomizePinData(UEdGraphPin* Pin, FName SourcePropertyName, int32 ArrayIndex) const override;
 	// End of UAnimGraphNode_Base interface
 	
+	// UObject override
+	void PostEditChangeProperty(FPropertyChangedEvent& PropertyChangedEvent) override; 	
+	
 private:
+	USkeletalMeshComponent* GetDebuggedComponent() const;
 	FAnimNode_RemapCurvesFromMesh* GetDebuggedNode() const;
-	
 };
diff --git a/Engine/Plugins/Experimental/Animation/CurveExpression/Source/Editor/Public/IRemapCurvesDebuggingProvider.h b/Engine/Plugins/Experimental/Animation/CurveExpression/Source/Editor/Public/IRemapCurvesDebuggingProvider.h
new file mode 100644
--- /dev/null
+++ b/Engine/Plugins/Experimental/Animation/CurveExpression/Source/Editor/Public/IRemapCurvesDebuggingProvider.h
@@ -0,0 +1,29 @@
+﻿// Copyright Epic Games, Inc. All Rights Reserved.
+
+#pragma once
+
+#include "UObject/Interface.h"
+
+#include "IRemapCurvesDebuggingProvider.generated.h"
+
+UINTERFACE()
+class URemapCurvesDebuggingProvider :
+	public UInterface
+{
+	GENERATED_BODY()
+};
+
+
+/** Helper interface to allow the details customization to work on both Remap Curves and Remap Curves from Mesh.
+  */
+class IRemapCurvesDebuggingProvider
+{
+	GENERATED_BODY()
+
+public:
+	/** Returns true if the expressions can be verified against an actively debugged anim graph instance */ 
+	virtual bool CanVerifyExpressions() const = 0;
+	
+	/** Trigger expression verification against an actively debugged anim graph instance */ 
+	virtual void VerifyExpressions() = 0;
+};
diff --git a/Engine/Plugins/Experimental/Animation/CurveExpression/Source/Editor/Public/K2Node_MakeCurveExpressionMap.h b/Engine/Plugins/Experimental/Animation/CurveExpression/Source/Editor/Public/K2Node_MakeCurveExpressionMap.h
--- a/Engine/Plugins/Experimental/Animation/CurveExpression/Source/Editor/Public/K2Node_MakeCurveExpressionMap.h
+++ b/Engine/Plugins/Experimental/Animation/CurveExpression/Source/Editor/Public/K2Node_MakeCurveExpressionMap.h
@@ -2,6 +2,8 @@
 
 #pragma once
 
+#include "CurveExpressionsDataAsset.h"
+
 #include "K2Node.h"
 
 #include "K2Node_MakeCurveExpressionMap.generated.h"
@@ -10,14 +12,6 @@ class FCompilerResultsLog;
 class FNodeHandlingFunctor;
 namespace ENodeTitleType { enum Type : int; }
 
-USTRUCT()
-struct FCurveExpressionList
-{
-	GENERATED_BODY()
-
-	UPROPERTY(EditAnywhere, Category="Expressions")
-	FString AssignmentExpressions;
-};
 
 UCLASS()
 class CURVEEXPRESSIONEDITOR_API UK2Node_MakeCurveExpressionMap :
diff --git a/Engine/Plugins/Experimental/Animation/CurveExpression/Source/Runtime/Private/AnimNode_RemapCurves.cpp b/Engine/Plugins/Experimental/Animation/CurveExpression/Source/Runtime/Private/AnimNode_RemapCurves.cpp
new file mode 100644
--- /dev/null
+++ b/Engine/Plugins/Experimental/Animation/CurveExpression/Source/Runtime/Private/AnimNode_RemapCurves.cpp
@@ -0,0 +1,68 @@
+﻿// Copyright Epic Games, Inc. All Rights Reserved.
+
+#include "AnimNode_RemapCurves.h"
+
+#include "Animation/AnimInstanceProxy.h"
+
+void FAnimNode_RemapCurves::Initialize_AnyThread(const FAnimationInitializeContext& Context)
+{
+	FAnimNode_RemapCurvesBase::Initialize_AnyThread(Context);
+
+	const USkeleton* Skeleton = Context.AnimInstanceProxy->GetSkeleton();
+	if (ensure(Skeleton))
+	{
+		const FSmartNameMapping* CurveNameMapping = Skeleton->GetSmartNameContainer(USkeleton::AnimCurveMappingName);
+
+		CurveNameMapping->Iterate([this](const FSmartNameMappingIterator& Iterator)
+		{
+			if (FName CurveName; Iterator.GetName(CurveName))
+			{
+				CurveNameToUIDMap.Add(CurveName, Iterator.GetIndex());
+			}
+		});
+	}
+	
+	CurveEvalResult.Reserve(GetCompiledAssignments().Num());
+}
+
+void FAnimNode_RemapCurves::Evaluate_AnyThread(FPoseContext& Output)
+{
+	DECLARE_SCOPE_HIERARCHICAL_COUNTER_ANIMNODE(Evaluate_AnyThread)
+	
+	using namespace CurveExpression::Evaluator;
+
+	FPoseContext SourceData(Output);
+	SourcePose.Evaluate(SourceData);
+	
+	CurveEvalResult.Reset();
+
+	for (const TTuple<FName, FExpressionObject>& Assignment: GetCompiledAssignments())
+	{
+		if (const SmartName::UID_Type* UID = CurveNameToUIDMap.Find(Assignment.Key))
+		{
+			CurveEvalResult.Add({*UID, FEngine().Execute(Assignment.Value,
+				[&CurveIDMap = CurveNameToUIDMap, &CurveValues = SourceData.Curve](const FName InName) -> TOptional<float>
+				{
+					if (const SmartName::UID_Type* SrcUID = CurveIDMap.Find(InName))
+					{
+						return CurveValues.Get(*SrcUID);
+					}
+					return {};
+				}
+			)});
+		}
+	}
+	
+	Output = SourceData;
+
+	for( const TTuple<SmartName::UID_Type, float>& CurveValue: CurveEvalResult)
+	{
+		Output.Curve.Set(CurveValue.Key, CurveValue.Value);
+	}
+}
+
+bool FAnimNode_RemapCurves::Serialize(FArchive& Ar)
+{
+	SerializeNode(Ar, this, StaticStruct());
+	return true;
+}
diff --git a/Engine/Plugins/Experimental/Animation/CurveExpression/Source/Runtime/Private/AnimNode_RemapCurvesFromMesh.cpp b/Engine/Plugins/Experimental/Animation/CurveExpression/Source/Runtime/Private/AnimNode_RemapCurvesFromMesh.cpp
--- a/Engine/Plugins/Experimental/Animation/CurveExpression/Source/Runtime/Private/AnimNode_RemapCurvesFromMesh.cpp
+++ b/Engine/Plugins/Experimental/Animation/CurveExpression/Source/Runtime/Private/AnimNode_RemapCurvesFromMesh.cpp
@@ -1,117 +1,64 @@
 // Copyright Epic Games, Inc. All Rights Reserved.
 
 #include "AnimNode_RemapCurvesFromMesh.h"
+
+#include "CurveExpressionCustomVersion.h"
 #include "Animation/AnimInstance.h"
 #include "Animation/AnimNodeFunctionRef.h"
 #include "Components/SkeletalMeshComponent.h"
 #include "Animation/ExposedValueHandler.h"
-#include "CurveExpressionModule.h"
 #include "Engine/SkeletalMesh.h"
 #include "Animation/AnimCurveTypes.h"
 
 #include UE_INLINE_GENERATED_CPP_BY_NAME(AnimNode_RemapCurvesFromMesh)
 
-void FAnimNode_RemapCurvesFromMesh::Initialize_AnyThread(
-	const FAnimationInitializeContext& Context
-	)
-{
-	DECLARE_SCOPE_HIERARCHICAL_COUNTER_ANIMNODE(Initialize_AnyThread)
-	Super::Initialize_AnyThread(Context);
-	SourcePose.Initialize(Context);
-}
 
-
-void FAnimNode_RemapCurvesFromMesh::CacheBones_AnyThread(
-	const FAnimationCacheBonesContext& Context
+void FAnimNode_RemapCurvesFromMesh::PreUpdate(
+	const UAnimInstance* InAnimInstance
 	)
 {
-	DECLARE_SCOPE_HIERARCHICAL_COUNTER_ANIMNODE(CacheBones_AnyThread)
-	Super::CacheBones_AnyThread(Context);
-	SourcePose.CacheBones(Context);
-}
-	
+	DECLARE_SCOPE_HIERARCHICAL_COUNTER_ANIMNODE(PreUpdate);
 
-void FAnimNode_RemapCurvesFromMesh::VerifyExpressions(
-	TFunction<void(FString)> InReportingFunc
-	) const
-{
-	using namespace CurveExpression::Evaluator;
+	// Make sure we're using the correct source and target skeleton components, since they
+	// may have changed from underneath us.
+	RefreshMeshComponent(InAnimInstance->GetSkelMeshComponent());
 
-	auto ReportAndLog = [InReportingFunc](FString InMessage)
-	{
-		if (InReportingFunc)
-		{
-			InReportingFunc(InMessage);
-		}
-		UE_LOG(LogCurveExpression, Warning, TEXT("%s"), *InMessage);
-	};
+	FAnimNode_RemapCurvesBase::PreUpdate(InAnimInstance);
 	
-	if (!ExpressionEngine.IsSet())
-	{
-		return;
-	}
-	
-	if (CurveExpressions.IsEmpty())
-	{
-		ReportAndLog(TEXT("No curve expressions set."));
-		return;
-	}
-
-	bool bFoundError = false;
-	const FEngine VerificationEngine(
-		ExpressionEngine->GetConstantValues(),
-		EParseFlags::ValidateConstants);
+	const USkeletalMeshComponent* CurrentMeshComponent = CurrentlyUsedSourceMeshComponent.IsValid() ? CurrentlyUsedSourceMeshComponent.Get() : nullptr;
 
-	for (const TPair<FName, FString>& ExpressionPair: CurveExpressions)
+	if (CurrentMeshComponent && CurrentMeshComponent->GetSkeletalMeshAsset() && CurrentMeshComponent->IsRegistered())
 	{
-		TOptional<FParseError> Error = VerificationEngine.Verify(ExpressionPair.Value);
-		if (Error.IsSet())
+		// If our source is running under leader-pose, then get bone data from there
+		if(const USkeletalMeshComponent* LeaderPoseComponent = Cast<USkeletalMeshComponent>(CurrentMeshComponent->LeaderPoseComponent.Get()))
 		{
-			ReportAndLog(FString::Printf(TEXT("Expression error in '%s': %s"), *ExpressionPair.Value, *Error->Message));
-			bFoundError = true;
+			CurrentMeshComponent = LeaderPoseComponent;
 		}
-	}
-
-	if (!bFoundError)
-	{
-		UE_LOG(LogCurveExpression, Display, TEXT("Curve expressions verified ok."))
-	}
-}
-
-
-
-bool FAnimNode_RemapCurvesFromMesh::CanVerifyExpressions() const
-{
-	return ExpressionEngine.IsSet();
-}
-
 
-void FAnimNode_RemapCurvesFromMesh::Update_AnyThread(
-	const FAnimationUpdateContext& Context
-	)
-{
-	DECLARE_SCOPE_HIERARCHICAL_COUNTER_ANIMNODE(Update_AnyThread)
-
-	// Run update on input pose nodes
-	SourcePose.Update(Context);
-
-	// Evaluate any BP logic plugged into this node
-	GetEvaluateGraphExposedInputs().Execute(Context);
-	
-	if (bExpressionsImmutable && ExpressionEngine.IsSet() && !CurveExpressions.IsEmpty() && CachedExpressions.IsEmpty())
-	{
-		CachedExpressions.Reset();
-		for (const TPair<FName, FString>& ExpressionPair: CurveExpressions)
+		// re-check mesh component validity as it may have changed to leader
+		if(CurrentMeshComponent->GetSkeletalMeshAsset() && CurrentMeshComponent->IsRegistered())
 		{
-			using namespace CurveExpression::Evaluator;
-
-			TVariant<FExpressionObject, FParseError> Result = ExpressionEngine->Parse(ExpressionPair.Value);
-			if (const FExpressionObject* Expression = Result.TryGet<FExpressionObject>())
+			if (const UAnimInstance* SourceAnimInstance = CurrentMeshComponent->GetAnimInstance())
 			{
-				CachedExpressions.Add(ExpressionPair.Key, *Expression);
+				// We have a valid instance, let's grab any curve values for the constants used by our expressions.
+				SourceCurveValues.Reset();
+
+				const TMap<FName, float>& CurveValues = SourceAnimInstance->GetAnimationCurveList(EAnimCurveType::AttributeCurve);
+				for (FName ConstantName: GetCompiledExpressionConstants())
+				{
+					if (const float* Value = CurveValues.Find(ConstantName))
+					{
+						SourceCurveValues.Add(ConstantName, *Value);
+					}
+				}
 			}
 		}
-	}			
+		else
+		{
+			// If there's no skeletal mesh, then empty the curve values.
+			SourceCurveValues.Reset();
+		}
+	}
 }
 
 
@@ -120,42 +67,28 @@ void FAnimNode_RemapCurvesFromMesh::Evaluate_AnyThread(
 	)
 {
 	DECLARE_SCOPE_HIERARCHICAL_COUNTER_ANIMNODE(Evaluate_AnyThread)
-
+	
+	using namespace CurveExpression::Evaluator;
+	
 	FPoseContext SourceData(Output);
 	SourcePose.Evaluate(SourceData);
 	
 	Output = SourceData;
-	
-	// If we have an expression engine, evaluate the expressions that have a matching target curve.
-	// If the expressions are not immutable between compiles, then we need to reparse them each time. 
-	if (ExpressionEngine.IsSet())
-	{
-		if (bExpressionsImmutable)
-		{
-			FBlendedCurve Curve;
-			Curve.Reserve(CachedExpressions.Num());
-			for (const TTuple<FName, CurveExpression::Evaluator::FExpressionObject>& ExpressionItem : CachedExpressions)
-			{
-				Curve.Add(ExpressionItem.Key, ExpressionEngine->Execute(ExpressionItem.Value)); 
-			}
 
-			Output.Curve.Combine(Curve);
-		}
-		else
+	for (const TTuple<FName, FExpressionObject>& Assignment: GetCompiledAssignments())
+	{
+		if (const SmartName::UID_Type* UID = CurveNameToUIDMap.Find(Assignment.Key))
 		{
-			FBlendedCurve Curve;
-			Curve.Reserve(CurveExpressions.Num());
-
-			for (const TTuple<FName, FString>& ExpressionItem : CurveExpressions)
-			{
-				TOptional<float> Result = ExpressionEngine->Evaluate(ExpressionItem.Value);
-				if (Result.IsSet())
+			Output.Curve.Set(*UID, FEngine().Execute(Assignment.Value,
+				[&CurveValues = SourceCurveValues](const FName InName) -> TOptional<float>
 				{
-					Curve.Add(ExpressionItem.Key, *Result);
+					if (const float* Value = CurveValues.Find(InName))
+					{
+						return *Value;
+					}
+					return {};
 				}
-			}
-
-			Output.Curve.Combine(Curve);
+			));
 		}
 	}
 }
@@ -173,42 +106,17 @@ void FAnimNode_RemapCurvesFromMesh::GatherDebugData(
 }
 
 
-void FAnimNode_RemapCurvesFromMesh::PreUpdate(
-	const UAnimInstance* InAnimInstance
-	)
+bool FAnimNode_RemapCurvesFromMesh::Serialize(FArchive& Ar)
 {
-	QUICK_SCOPE_CYCLE_COUNTER(FAnimNode_RemapCurvesFromMesh_PreUpdate);
+	SerializeNode(Ar, this, StaticStruct());
 
-	// Make sure we're using the correct source and target skeleton components, since they
-	// may have changed from underneath us.
-	RefreshMeshComponent(InAnimInstance->GetSkelMeshComponent());
-
-	const USkeletalMeshComponent* CurrentMeshComponent = CurrentlyUsedSourceMeshComponent.IsValid() ? CurrentlyUsedSourceMeshComponent.Get() : nullptr;
-
-	if (CurrentMeshComponent && CurrentMeshComponent->GetSkeletalMeshAsset() && CurrentMeshComponent->IsRegistered())
+	if (IsLoading() && Ar.CustomVer(FCurveExpressionCustomVersion::GUID) < FCurveExpressionCustomVersion::SerializedExpressions)
 	{
-		// If our source is running under leader-pose, then get bone data from there
-		if(USkeletalMeshComponent* LeaderPoseComponent = Cast<USkeletalMeshComponent>(CurrentMeshComponent->LeaderPoseComponent.Get()))
-		{
-			CurrentMeshComponent = LeaderPoseComponent;
-		}
-
-		// re-check mesh component validity as it may have changed to leader
-		if(CurrentMeshComponent->GetSkeletalMeshAsset() && CurrentMeshComponent->IsRegistered())
-		{
-			if (ExpressionEngine.IsSet())
-			{
-				if (const UAnimInstance* SourceAnimInstance = CurrentMeshComponent->GetAnimInstance())
-				{
-					ExpressionEngine->UpdateConstantValues(SourceAnimInstance->GetAnimationCurveList(EAnimCurveType::AttributeCurve));
-				}
-			}
-		}
-		else
-		{
-			ExpressionEngine.Reset();
-		}
+		// Default to expression map if loading from an older version.
+		ExpressionSource = ERemapCurvesExpressionSource::ExpressionMap;
 	}
+	
+	return true;
 }
 
 
@@ -221,8 +129,6 @@ void FAnimNode_RemapCurvesFromMesh::ReinitializeMeshComponent(
 	CurrentlyUsedSourceMesh.Reset();
 	CurrentlyUsedTargetMesh.Reset();
 	
-	ExpressionEngine.Reset();
-	CachedExpressions.Reset();
 
 	if (InTargetMeshComponent && IsValid(InNewSkeletalMeshComponent) && InNewSkeletalMeshComponent->GetSkeletalMeshAsset())
 	{
@@ -236,19 +142,6 @@ void FAnimNode_RemapCurvesFromMesh::ReinitializeMeshComponent(
 			CurrentlyUsedSourceMesh = SourceSkelMesh;
 			CurrentlyUsedTargetMesh = TargetSkelMesh;
 
-			// Grab the source curves and use their names to seed the expression evaluation engine.
-			TMap<FName, float> SourceCurves;
-			const USkeleton* SourceSkeleton = SourceSkelMesh->GetSkeleton();
-			if (ensureMsgf(SourceSkeleton, TEXT("Invalid null source skeleton : %s"), *GetNameSafe(TargetSkelMesh)))
-			{
-				// TODO: is this good enough for future use? ExpressionEngine appears to have been built on the assumption that curves are always pre-defined.
-				SourceSkeleton->ForEachCurveMetaData([&SourceCurves](FName InCurveName, const FCurveMetaData& InCurveMetaData)
-				{
-					SourceCurves.Add(InCurveName, 0.0f);
-				});
-
-				ExpressionEngine.Emplace(MoveTemp(SourceCurves));
-			}
 		}
 	}
 }
diff --git a/Engine/Plugins/Experimental/Animation/CurveExpression/Source/Runtime/Private/CurveExpressionCustomVersion.cpp b/Engine/Plugins/Experimental/Animation/CurveExpression/Source/Runtime/Private/CurveExpressionCustomVersion.cpp
new file mode 100644
--- /dev/null
+++ b/Engine/Plugins/Experimental/Animation/CurveExpression/Source/Runtime/Private/CurveExpressionCustomVersion.cpp
@@ -0,0 +1,12 @@
+﻿// Copyright Epic Games, Inc. All Rights Reserved.
+
+#include "CurveExpressionCustomVersion.h"
+
+#include "UObject/DevObjectVersion.h"
+
+
+const FGuid FCurveExpressionCustomVersion::GUID(0xA26D36AE, 0x26935388, 0xA8C5CB96, 0x2B95B4AF);
+
+// Register the custom version with core
+FCustomVersionRegistration GRegisterCurveExpressionCustomVersion(
+	FCurveExpressionCustomVersion::GUID, FCurveExpressionCustomVersion::LatestVersion, TEXT("CurveExpressionVersion"));
diff --git a/Engine/Plugins/Experimental/Animation/CurveExpression/Source/Runtime/Private/CurveExpressionCustomVersion.h b/Engine/Plugins/Experimental/Animation/CurveExpression/Source/Runtime/Private/CurveExpressionCustomVersion.h
new file mode 100644
--- /dev/null
+++ b/Engine/Plugins/Experimental/Animation/CurveExpression/Source/Runtime/Private/CurveExpressionCustomVersion.h
@@ -0,0 +1,25 @@
+﻿// Copyright Epic Games, Inc. All Rights Reserved.
+
+#pragma once
+
+#include "Misc/Guid.h"
+
+struct FCurveExpressionCustomVersion
+{
+	FCurveExpressionCustomVersion() = delete;
+	
+	enum Type
+	{
+		// Before any version changes were made in niagara
+		BeforeCustomVersionWasAdded = 0,
+
+		// Serialized expressions
+		SerializedExpressions,
+		
+		VersionPlusOne,
+		LatestVersion = VersionPlusOne - 1,
+	};
+	
+	// The GUID for this custom version number
+	const static FGuid GUID;
+};
diff --git a/Engine/Plugins/Experimental/Animation/CurveExpression/Source/Runtime/Private/CurveExpressionsDataAsset.cpp b/Engine/Plugins/Experimental/Animation/CurveExpression/Source/Runtime/Private/CurveExpressionsDataAsset.cpp
new file mode 100644
--- /dev/null
+++ b/Engine/Plugins/Experimental/Animation/CurveExpression/Source/Runtime/Private/CurveExpressionsDataAsset.cpp
@@ -0,0 +1,157 @@
+﻿// Copyright Epic Games, Inc. All Rights Reserved.
+
+#include "CurveExpressionsDataAsset.h"
+
+#if WITH_EDITOR
+#include "ExpressionEvaluator.h"
+#endif
+
+#include "String/ParseLines.h"
+
+#include UE_INLINE_GENERATED_CPP_BY_NAME(CurveExpressionsDataAsset)
+
+
+TArray<FCurveExpressionAssignment> FCurveExpressionList::GetAssignments() const
+{
+	TArray<FCurveExpressionAssignment> ExpressionAssignments;
+	int32 LineIndex = 0;
+	
+	UE::String::ParseLines(AssignmentExpressions,
+		[&ExpressionAssignments, &LineIndex](FStringView InLine)
+	{
+		int32 AssignmentPos;
+		if (InLine.FindChar('=', AssignmentPos))
+		{
+			const FStringView AssignmentTarget = InLine.Left(AssignmentPos).TrimStartAndEnd();
+			const FStringView SourceExpression = InLine.Mid(AssignmentPos + 1).TrimStartAndEnd();
+			if (!AssignmentTarget.IsEmpty() && !SourceExpression.IsEmpty())
+			{
+				ExpressionAssignments.Add({LineIndex, FName(AssignmentTarget), SourceExpression});
+			}
+		}
+		LineIndex++;
+	});
+	return ExpressionAssignments;
+}
+
+TArray<FCurveExpressionParsedAssignment> FCurveExpressionList::GetParsedAssignments() const
+{
+	using namespace CurveExpression::Evaluator;
+	
+	TArray<FCurveExpressionParsedAssignment> ParsedAssignments;
+	FEngine ExpressionEngine;
+
+	for (const FCurveExpressionAssignment& Assignment: GetAssignments())
+	{
+		TVariant<FExpressionObject, FParseError> Result = ExpressionEngine.Parse(Assignment.Expression);
+		ParsedAssignments.Add({Assignment.LineIndex, Assignment.TargetName, MoveTemp(Result)});
+	}
+	
+	return ParsedAssignments;
+}
+
+struct FCompiledExpressionScratchArea : public TThreadSingleton<FCompiledExpressionScratchArea>
+{
+	// Shadowed data copy for extraction from non-game threads.
+	uint32 SerialNumber{UINT_MAX};
+	TArray<FName> NamedConstants;
+	TMap<FName, CurveExpression::Evaluator::FExpressionObject> ExpressionMap;
+};
+
+
+void UCurveExpressionsDataAsset::SynchronizeThreadLocalData() const
+{
+	FCompiledExpressionScratchArea& ThreadLocalData = FCompiledExpressionScratchArea::Get();
+		
+	const uint32 Serial = SerialNumber.load();
+	if (ThreadLocalData.SerialNumber != Serial)
+	{
+		FScopeLock ScopeLock(&ExpressionMapWriteLock);
+
+		ThreadLocalData.SerialNumber = Serial;
+		ThreadLocalData.NamedConstants = NamedConstants;
+		ThreadLocalData.ExpressionMap = ExpressionMap;
+	}
+}
+
+
+const TArray<FName>& UCurveExpressionsDataAsset::GetCompiledExpressionConstants() const
+{
+	if (IsInGameThread())
+	{
+		return NamedConstants;
+	}
+	else
+	{
+		SynchronizeThreadLocalData();
+		return FCompiledExpressionScratchArea::Get().NamedConstants;
+	}
+}
+
+
+const TMap<FName, CurveExpression::Evaluator::FExpressionObject>& UCurveExpressionsDataAsset::GetCompiledExpressionMap() const
+{
+	if (IsInGameThread())
+	{
+		return ExpressionMap;
+	}
+	else
+	{
+		SynchronizeThreadLocalData();
+		return FCompiledExpressionScratchArea::Get().ExpressionMap;
+	}
+}
+
+
+void UCurveExpressionsDataAsset::Serialize(FArchive& Ar)
+{
+	UObject::Serialize(Ar);
+
+#if WITH_EDITOR
+	if (Ar.IsSaving())
+	{
+		// Make sure the expressions are compiled.
+		CompileExpressions();
+	}
+#endif
+
+	// Serialize the compiled map so we can read it in cooked builds.
+	Ar << ExpressionMap;
+}
+
+
+#if WITH_EDITOR
+void UCurveExpressionsDataAsset::PostEditChangeProperty(FPropertyChangedEvent& PropertyChangedEvent)
+{
+	Super::PostEditChangeProperty(PropertyChangedEvent);
+
+	if (PropertyChangedEvent.GetPropertyName() == GET_MEMBER_NAME_CHECKED(FCurveExpressionList, AssignmentExpressions))
+	{
+		CompileExpressions();
+	}
+}
+
+
+void UCurveExpressionsDataAsset::CompileExpressions()
+{
+	using namespace CurveExpression::Evaluator;
+	
+	FScopeLock ScopeLock(&ExpressionMapWriteLock);
+
+	TSet<FName> ConstantNames;
+	
+	ExpressionMap.Reset();
+	for (const FCurveExpressionParsedAssignment& Assignment: Expressions.GetParsedAssignments())
+	{
+		if (const FExpressionObject* Expression = Assignment.Result.TryGet<FExpressionObject>())
+		{
+			ConstantNames.Append(Expression->GetUsedConstants());
+			ExpressionMap.Add(Assignment.TargetName, *Expression);
+		}
+	}
+	
+	NamedConstants = ConstantNames.Array();
+	SerialNumber += 1;
+}
+
+#endif
diff --git a/Engine/Plugins/Experimental/Animation/CurveExpression/Source/Runtime/Private/ExpressionEvaluator.cpp b/Engine/Plugins/Experimental/Animation/CurveExpression/Source/Runtime/Private/ExpressionEvaluator.cpp
--- a/Engine/Plugins/Experimental/Animation/CurveExpression/Source/Runtime/Private/ExpressionEvaluator.cpp
+++ b/Engine/Plugins/Experimental/Animation/CurveExpression/Source/Runtime/Private/ExpressionEvaluator.cpp
@@ -3,6 +3,8 @@
 #include "ExpressionEvaluator.h"
 #include <limits> // IWYU pragma: keep
 
+#include "CurveExpressionModule.h"
+
 namespace CurveExpression::Evaluator
 {
 
@@ -104,63 +106,124 @@ FString FExpressionObject::ToString() const
 }
 
 
-FEngine::FEngine(
-	const TMap<FName, float>& InConstants, 
-	EParseFlags InParseFlags
-	) :
-	Constants(InConstants),
-	ParseFlags(InParseFlags)
+TSet<FName> FExpressionObject::GetUsedConstants() const
 {
+	TSet<FName> UsedConstants;
+	for (int32 Index = 0; Index < Expression.Num(); Index++)
+	{
+		if (const FName* ConstantName = Expression[Index].TryGet<FName>())
+		{
+			UsedConstants.Add(*ConstantName);
+		}
+	}
+	return UsedConstants;	
 }
 
 
-FEngine::FEngine(
-	TMap<FName, float>&& InConstants,
-	EParseFlags InParseFlags
-	) :
-	Constants(MoveTemp(InConstants)),
-	ParseFlags(InParseFlags)
+void FExpressionObject::Load(FArchive& Ar)
 {
-}
+	int32 OperandCount = 0;
+	Ar << OperandCount;
 
+	Expression.Reset(OperandCount);
 
-void FEngine::UpdateConstantValues(
-	const TMap<FName, float>& InConstants,
-	EUpdateConstantsMethod InUpdateMethod
-	)
-{
-	for (TTuple<FName, float>& OldConstantItem: Constants)
+	for (int32 OperandIndex = 0; OperandIndex < OperandCount; OperandIndex++)
 	{
-		if (const float* NewConstant = InConstants.Find(OldConstantItem.Key))
+		int32 OperandType = INDEX_NONE;
+		Ar << OperandType;
+		switch(OperandType)
 		{
-			OldConstantItem.Value = *NewConstant;
+		case OpElement::IndexOfType<EOperator>():
+			{
+				int32 OperatorType = 0;
+				Ar << OperatorType;
+				Expression.Add(OpElement(TInPlaceType<EOperator>(), static_cast<EOperator>(OperatorType)));
+			}
+			break;
+		case OpElement::IndexOfType<FName>():
+			{
+				FName ConstantName;
+				Ar << ConstantName;
+				Expression.Add(OpElement(TInPlaceType<FName>(), ConstantName));
+			}
+			break;
+		case OpElement::IndexOfType<FFunctionRef>():
+			{
+				int32 FunctionIndex;
+				Ar << FunctionIndex;
+				Expression.Add(OpElement(TInPlaceType<FFunctionRef>(), FFunctionRef(FunctionIndex)));
+			}
+			break;
+		case OpElement::IndexOfType<float>():
+			{
+				float Value;
+				Ar << Value;
+				Expression.Add(OpElement(TInPlaceType<float>(), Value));
+			}
+			break;
+		default:
+			UE_LOG(LogCurveExpression, Error, TEXT("Invalid element type found in serialized expression"));
+			Ar.SetError();
+			return;
 		}
-		else if (InUpdateMethod == EUpdateConstantsMethod::ZeroMissing)
+	}
+}
+
+
+void FExpressionObject::Save(FArchive& Ar) const
+{
+	int32 OperandCount = Expression.Num(); 
+	Ar << OperandCount;
+
+	for (const OpElement& Operand: Expression)
+	{
+		int32 OperandType = Operand.GetIndex();
+		Ar << OperandType;
+		
+		switch(OperandType)
 		{
-			OldConstantItem.Value = 0.0f;
+		case OpElement::IndexOfType<EOperator>():
+			{
+				int32 OperatorType = static_cast<int32>(Operand.Get<EOperator>());
+				Ar << OperatorType;
+			}
+			break;
+		case OpElement::IndexOfType<FName>():
+			{
+				FName ConstantName = Operand.Get<FName>();
+				Ar << ConstantName;
+			}
+			break;
+		case OpElement::IndexOfType<FFunctionRef>():
+			{
+				int32 FunctionIndex = Operand.Get<FFunctionRef>().Index;
+				Ar << FunctionIndex;
+			}
+			break;
+		case OpElement::IndexOfType<float>():
+			{
+				float Value = Operand.Get<float>();
+				Ar << Value;
+			}
+			break;
+		default:
+			checkNoEntry();
+			return;
 		}
 	}
 }
 
+
 // We can't use FChar::IsDigit, since it obeys locale, which would allow all sorts of non-ASCII
 // numeric digits from Unicode.
-static bool IsASCIIDigit(TCHAR InCh)
+static bool IsASCIIDigit(const TCHAR InCh)
 {
 	return InCh >= '0' && InCh <= '9';
 }
 
-TVariant<FEngine::FToken, FParseError> FEngine::ParseTokenError(
-	FString&& InError,
-	const TCHAR* InCurrentPos,
-	FStringView InExpression)
-{
-	return TVariant<FToken, FParseError>(TInPlaceType<FParseError>(),
-		static_cast<int32>(InCurrentPos - InExpression.GetData()), MoveTemp(InError));
-}
-
 
 FString FEngine::TokenToString(
-	const FEngine::FToken& InToken
+	const FToken& InToken
 	)
 {
 	const TCHAR* TypeStr = TEXT("<unknown>");
@@ -182,6 +245,10 @@ FString FEngine::TokenToString(
 		case EOperatorToken::ParenOpen:		ValueStr = TEXT("("); break;
 		case EOperatorToken::ParenClose:	ValueStr = TEXT(")"); break;
 		case EOperatorToken::Comma:			ValueStr = TEXT(","); break;
+		
+		case EOperatorToken::Max:
+			checkNoEntry();
+			break;
 		}
 	}
 	else if (const FName* Identifier = InToken.TryGet<FName>())
@@ -199,10 +266,10 @@ FString FEngine::TokenToString(
 }
 
 
-TVariant<FEngine::FToken, FParseError> FEngine::ParseFloat(
+auto FEngine::ParseFloat(
 	FStringView& InOutParseRange,
 	FStringView InExpression
-	)
+	) -> FTokenParseResult
 {
 	// We only parse at most 11 significant digits, since 9 is the maximum decimal precision 
 	// a float will hold. We store 2 extra digits to make sure we round properly on the 
@@ -258,7 +325,9 @@ TVariant<FEngine::FToken, FParseError> FEngine::ParseFloat(
 	if (MantissaLen == 0)
 	{
 		// We only parsed the period. That's not a valid float.
-		return ParseTokenError(TEXT("Invalid floating point value"), InOutParseRange.GetData(), InExpression);
+		return FTokenParseResult(TEXT("Invalid floating point value"),
+			InOutParseRange.Left(1), 
+		InExpression.GetData());
 	}
 
 	// Adjustment for the exponent based on number of significant digits in the fraction part.
@@ -372,17 +441,19 @@ TVariant<FEngine::FToken, FParseError> FEngine::ParseFloat(
 		Fraction *= DoubleExponent;
 	}
 
+	const FStringView TokenParseRange(InOutParseRange.Left(LastPos));
+	
+	// Update the parse range for the next token. 
 	InOutParseRange = InOutParseRange.Mid(LastPos);
 	
-	return TVariant<FToken, FParseError>(
-		TInPlaceType<FToken>(), FToken(TInPlaceType<float>(), Fraction));
+	return FTokenParseResult(FToken(TInPlaceType<float>(), Fraction), TokenParseRange, InExpression.GetData());
 }
 
 
-TVariant<FEngine::FToken, FParseError> FEngine::ParseIdentifier(
+auto FEngine::ParseIdentifier(
 	FStringView& InOutParseRange,
 	FStringView InExpression
-	) const
+	) -> FTokenParseResult 
 {
 	// FIXME: We need quote escaping.
 	bool bStopAtSingleQuote = false;
@@ -393,13 +464,13 @@ TVariant<FEngine::FToken, FParseError> FEngine::ParseIdentifier(
 		ParseIndex++;
 		if (ParseIndex >= InOutParseRange.Len())
 		{
-			return ParseTokenError(TEXT("Missing end quote for constant"), InOutParseRange.GetData(), InExpression);
+			return FTokenParseResult(TEXT("Missing end quote for constant"), InOutParseRange, InExpression.GetData());
 		}
 	}
 	// Unquoted constants have to start with an alphabetic character.
 	else if (!FChar::IsAlpha(InOutParseRange[0]) && InOutParseRange[0] != '_')
 	{
-		return ParseTokenError(TEXT("Unexpected character"), InOutParseRange.GetData(), InExpression);
+		return FTokenParseResult(TEXT("Unexpected character"), InOutParseRange.Left(1), InExpression.GetData());
 	}
 
 	const int32 ConstantStart = ParseIndex;
@@ -416,7 +487,7 @@ TVariant<FEngine::FToken, FParseError> FEngine::ParseIdentifier(
 
 		if (ParseIndex == InOutParseRange.Len())
 		{
-			return ParseTokenError(TEXT("Missing end quote for constant"), InOutParseRange.GetData(), InExpression);
+			return FTokenParseResult(TEXT("Missing end quote for constant"), InOutParseRange, InExpression.GetData());
 		}
 
 		ConstantEnd = ParseIndex;
@@ -433,86 +504,61 @@ TVariant<FEngine::FToken, FParseError> FEngine::ParseIdentifier(
 		ConstantEnd = ParseIndex;
 	}
 
-	// Look up this identifier. We don't want to create a new FName, to avoid polluting the
-	// FNameEntry cache.
 	const FStringView Identifier(InOutParseRange.SubStr(ConstantStart, ConstantEnd - ConstantStart));
-	const FName Name(Identifier.Len(), Identifier.GetData(), FNAME_Find);
+	const FName Name(Identifier.Len(), Identifier.GetData());
 	
-	InOutParseRange = InOutParseRange.Mid(ParseIndex); 
+	const FStringView TokenParseRange(InOutParseRange.Left(ParseIndex));
+	InOutParseRange = InOutParseRange.Mid(ParseIndex);
 
-	if (Constants.Contains(Name) || FunctionNameIndex.Contains(Name))
-	{
-		return TVariant<FToken, FParseError>(
-			TInPlaceType<FToken>(), FToken(TInPlaceType<FName>(), Name));
-	}
-	else if (EnumHasAllFlags(ParseFlags, EParseFlags::ValidateConstants))
-	{
-		return ParseTokenError(FString::Printf(TEXT("Unknown identifier '%s'"), *FString(Identifier)),
-			InOutParseRange.GetData(), InExpression); 
-	}
-	else
-	{
-		// If the constant value is not found, it gets a default value of None.
-		return TVariant<FToken, FParseError>(
-			TInPlaceType<FToken>(), FToken(TInPlaceType<FName>(), NAME_None));
-	}
+	return FTokenParseResult(FToken(TInPlaceType<FName>(), Name), TokenParseRange, InExpression.GetData());
 }
 
 
-TVariant<FEngine::FToken, FParseError> FEngine::ParseToken(
+auto FEngine::ParseToken(
 	FStringView& InOutParseRange,
 	FStringView InExpression
-	) const
+	) -> FTokenParseResult
 {
-	auto TokenResult = [](EOperatorToken InTokenKind)
+	auto TokenResult = [&InOutParseRange, Expression=InExpression.GetData()](EOperatorToken InTokenKind, int32 InTokenLength)
 	{
-		return TVariant<FToken, FParseError>(
-			TInPlaceType<FToken>(), FToken(TInPlaceType<EOperatorToken>(), InTokenKind));
+		const FStringView TokenRange(InOutParseRange.Left(InTokenLength));
+		InOutParseRange = InOutParseRange.Mid(InTokenLength); 
+		return FTokenParseResult(FToken(TInPlaceType<EOperatorToken>(), InTokenKind), TokenRange, Expression);
 	};
 
 	switch(InOutParseRange[0])
 	{
 	case '+':
-		InOutParseRange = InOutParseRange.Mid(1);
-		return TokenResult(EOperatorToken::Add);
+		return TokenResult(EOperatorToken::Add, 1);
 		
 	case '-':
-		InOutParseRange = InOutParseRange.Mid(1);
-		return TokenResult(EOperatorToken::Subtract);
+		return TokenResult(EOperatorToken::Subtract, 1);
 		
 	case '*':
 		if (InOutParseRange.Len() >= 2 && InOutParseRange[1] == '*')
 		{
-			InOutParseRange = InOutParseRange.Mid(2);
-			return TokenResult(EOperatorToken::Power);
+			return TokenResult(EOperatorToken::Power, 2);
 		}
-		InOutParseRange = InOutParseRange.Mid(1);
-		return TokenResult(EOperatorToken::Multiply);
+		return TokenResult(EOperatorToken::Multiply, 1);
 		
 	case '/':
 		if (InOutParseRange.Len() >= 2 && InOutParseRange[1] == '/')
 		{
-			InOutParseRange = InOutParseRange.Mid(2);
-			return TokenResult(EOperatorToken::FloorDivide);
+			return TokenResult(EOperatorToken::FloorDivide, 2);
 		}
-		InOutParseRange = InOutParseRange.Mid(1);
-		return TokenResult(EOperatorToken::Divide);
+		return TokenResult(EOperatorToken::Divide, 1);
 		
 	case '%':
-		InOutParseRange = InOutParseRange.Mid(1);
-		return TokenResult(EOperatorToken::Modulo);
+		return TokenResult(EOperatorToken::Modulo, 1);
 
 	case '(':
-		InOutParseRange = InOutParseRange.Mid(1);
-		return TokenResult(EOperatorToken::ParenOpen);
+		return TokenResult(EOperatorToken::ParenOpen, 1);
 		
 	case ')':
-		InOutParseRange = InOutParseRange.Mid(1);
-		return TokenResult(EOperatorToken::ParenClose);
+		return TokenResult(EOperatorToken::ParenClose, 1);
 
 	case ',':
-		InOutParseRange = InOutParseRange.Mid(1);
-		return TokenResult(EOperatorToken::Comma);
+		return TokenResult(EOperatorToken::Comma, 1);
 		
 	default:
 		if (IsASCIIDigit(InOutParseRange[0]) || InOutParseRange[0] == '.')
@@ -525,7 +571,8 @@ TVariant<FEngine::FToken, FParseError> FEngine::ParseToken(
 
 
 TVariant<FExpressionObject, FParseError> FEngine::Parse(
-	FStringView InExpression
+	FStringView InExpression,
+	TFunction<TOptional<float>(FName InConstantName)> InConstantEvaluator
 	) const
 {
 	enum class EAssociativity
@@ -573,44 +620,44 @@ TVariant<FExpressionObject, FParseError> FEngine::Parse(
 	// Run Dijkstra's classic Shunting Yard algorithm to convert infix expressions to RPN.
 	// TODO: Better error pinpointing by storing the token parse result with the token.
 	TArray<FExpressionObject::OpElement, TInlineAllocator<64>> Expression;
-	TArray<EOperatorToken, TInlineAllocator<32>> OperatorStack;
+	TArray<FParseLocation, TInlineAllocator<64>> Locations;
+	TArray<TTuple<EOperatorToken, FParseLocation>, TInlineAllocator<32>> OperatorStack;
+	
 	struct FFunctionCallInfo
 	{
 		int32 FunctionIndex = INDEX_NONE;
 		int32 CountedCommas = 0;
 		int32 OpeningOpStackSize = 0;
 		int32 ExpressionSize = 0;
-		const TCHAR* TokenStart = nullptr;
+		FParseLocation Location{};
 	};
 	TArray<FFunctionCallInfo, TInlineAllocator<32>> FunctionStack;
 
-	auto ParseError = [InExpression](FString&& InError, const TCHAR* InCurrentPos)
+	auto ParseError = [](FString&& InError, FParseLocation InLocation)
 	{
-		return TVariant<FExpressionObject, FParseError>(TInPlaceType<FParseError>(),
-			static_cast<int32>(InCurrentPos - InExpression.GetData()), MoveTemp(InError));
+		return TVariant<FExpressionObject, FParseError>(TInPlaceType<FParseError>(), MoveTemp(InError), InLocation);
 	};
-	auto PushOperatorTokenToExpression = [GetOperatorTokenInfo, &Expression](EOperatorToken InOpToken)
+	auto PushOperatorTokenToExpression = [GetOperatorTokenInfo, &Expression, &Locations](TTuple<EOperatorToken, FParseLocation> InTokenAndLocation)
 	{
 		Expression.Push(FExpressionObject::OpElement(TInPlaceType<FExpressionObject::EOperator>(), 
-			*GetOperatorTokenInfo(InOpToken).Operator));
+			*GetOperatorTokenInfo(InTokenAndLocation.Get<0>()).Operator));
+		Locations.Push(InTokenAndLocation.Get<1>());
 	};
 	// Used to check if we encounter a '-' to see if we need to turn it into a negate operator
 	// and also to see if we have two operators in a row (error).
 	TOptional<FToken> LastToken;
-	FStringView LastTokenParseRange;
 
 	do
 	{
-		const TCHAR* TokenStart = ParseRange.GetData();
-		TVariant<FToken, FParseError> ParseResult = ParseToken(ParseRange, InExpression);
+		FTokenParseResult ParseResult = ParseToken(ParseRange, InExpression);
 		
-		if (FParseError* Error = ParseResult.TryGet<FParseError>())
+		if (ParseResult.IsError())
 		{
-			return TVariant<FExpressionObject, FParseError>(TInPlaceType<FParseError>(), *Error);
+			return TVariant<FExpressionObject, FParseError>(TInPlaceType<FParseError>(), ParseResult.GetParseError());
 		}
 
-		const FStringView TokenParseRange(TokenStart, ParseRange.GetData() - TokenStart);
-		const FToken& Token = ParseResult.Get<FToken>();
+		const FToken& Token = ParseResult.GetToken();
+		
 
 		// Values / identifiers go onto the expression stack immediately.
 		if (const FName* Identifier = Token.TryGet<FName>())
@@ -618,18 +665,20 @@ TVariant<FExpressionObject, FParseError> FEngine::Parse(
 			if (LastToken.IsSet() &&
 				(LastToken->IsType<FName>() || LastToken->IsType<float>()))
 			{
-				return ParseError(TEXT("Expected an operator"), TokenParseRange.GetData());
+				return ParseError(TEXT("Expected an operator"), ParseResult.Location);
 			}
 			Expression.Push(FExpressionObject::OpElement(TInPlaceType<FName>(), *Identifier));
+			Locations.Push(ParseResult.Location);
 		}
 		else if (const float* Value = Token.TryGet<float>())
 		{
 			if (LastToken.IsSet() &&
 				(LastToken->IsType<FName>() || LastToken->IsType<float>()))
 			{
-				return ParseError(TEXT("Expected an operator"), TokenParseRange.GetData());
+				return ParseError(TEXT("Expected an operator"), ParseResult.Location);
 			}
 			Expression.Push(FExpressionObject::OpElement(TInPlaceType<float>(), *Value));
+			Locations.Push(ParseResult.Location);
 		}
 		else
 		{
@@ -644,16 +693,16 @@ TVariant<FExpressionObject, FParseError> FEngine::Parse(
 					const FOperatorTokenInfo& LastTokenInfo = GetOperatorTokenInfo(*LastOp);
 					if (TokenInfo.Operator.IsSet() && LastTokenInfo.Operator.IsSet())
 					{
-						return ParseError(TEXT("Expected an expression"), TokenParseRange.GetData());
+						return ParseError(TEXT("Expected an expression"), ParseResult.Location);
 					}
 					if (*LastOp == EOperatorToken::ParenClose && Op == EOperatorToken::ParenOpen)
 					{
-						return ParseError(TEXT("Expected an operator"), TokenParseRange.GetData());
+						return ParseError(TEXT("Expected an operator"), ParseResult.Location);
 					}
 					if (*LastOp == EOperatorToken::ParenOpen && Op == EOperatorToken::ParenClose &&
 					    (FunctionStack.IsEmpty() || FunctionStack.Top().OpeningOpStackSize != OperatorStack.Num()))
 					{
-						return ParseError(TEXT("Empty Parentheses"), TokenParseRange.GetData());
+						return ParseError(TEXT("Empty Parentheses"), ParseResult.Location);
 					}
 				}
 			}
@@ -663,34 +712,35 @@ TVariant<FExpressionObject, FParseError> FEngine::Parse(
 				(!LastToken.IsSet() || (LastToken->IsType<EOperatorToken>() && LastToken->Get<EOperatorToken>() != EOperatorToken::ParenClose)))
 			{
 				Op = EOperatorToken::Negate;
-				OperatorStack.Push(Op);
+				OperatorStack.Push({Op, ParseResult.Location});
 			}
 			else if (Op == EOperatorToken::ParenOpen)
 			{
-				OperatorStack.Push(Op);
+				OperatorStack.Push({Op, ParseResult.Location});
 
-				// If the top element on the expression is a constant, check if it is a function and if so, pop the constant
-				// out of the expression. 
+				// If the top element on the expression is a constant, check if it is a function name and if so, 
+				// pop the constant out of the expression. 
 				if (!Expression.IsEmpty() && Expression.Top().IsType<FName>())
 				{
 					FName ConstantName = Expression.Top().Get<FName>();
 					const int32 *FunctionIndexPtr = FunctionNameIndex.Find(ConstantName);
 					if (FunctionIndexPtr == nullptr)
 					{
-						return ParseError(FString::Printf(TEXT("Unknown function '%s'"), *ConstantName.ToString()), LastTokenParseRange.GetData());
+						return ParseError(FString::Printf(TEXT("Unknown function '%s'"), *ConstantName.ToString()), Locations.Top());
 					}
 
 					Expression.Pop();
-					FunctionStack.Push({*FunctionIndexPtr, 0, OperatorStack.Num(), Expression.Num(), LastTokenParseRange.GetData()});
+					const FParseLocation LastLocation = Locations.Pop();
+					FunctionStack.Push({*FunctionIndexPtr, 0, OperatorStack.Num(), Expression.Num(), LastLocation});
 				}
 			}
 			else if (Op == EOperatorToken::ParenClose)
 			{
-				while(OperatorStack.IsEmpty() || OperatorStack.Last() != EOperatorToken::ParenOpen)
+				while(OperatorStack.IsEmpty() || OperatorStack.Last().Key != EOperatorToken::ParenOpen)
 				{
 					if (OperatorStack.IsEmpty())
 					{
-						return ParseError(TEXT("Mismatched parentheses"), TokenParseRange.GetData()); 
+						return ParseError(TEXT("Mismatched parentheses"), ParseResult.Location); 
 					}
 					PushOperatorTokenToExpression(OperatorStack.Pop(false));
 				}
@@ -704,10 +754,12 @@ TVariant<FExpressionObject, FParseError> FEngine::Parse(
 
 					if (ArgumentCount != FunctionInfo.ArgumentCount)
 					{
-						return ParseError(FString::Printf(TEXT("Invalid argument count. Expected %d, got %d"), FunctionInfo.ArgumentCount, ArgumentCount), CallInfo.TokenStart); 
+						return ParseError(FString::Printf(TEXT("Invalid argument count. Expected %d, got %d"), FunctionInfo.ArgumentCount, ArgumentCount),
+							CallInfo.Location.Join(ParseResult.Location)); 
 					}
 
 					Expression.Push(FExpressionObject::OpElement(TInPlaceType<FExpressionObject::FFunctionRef>(), CallInfo.FunctionIndex));
+					Locations.Push(ParseResult.Location);
 				}
 				
 				// Remove the now unneeded open parentheses.
@@ -717,44 +769,58 @@ TVariant<FExpressionObject, FParseError> FEngine::Parse(
 			{
 				if (FunctionStack.IsEmpty() || FunctionStack.Top().OpeningOpStackSize != OperatorStack.Num())
 				{
-					return ParseError(TEXT("Unexpected comma"), TokenParseRange.GetData()); 
+					return ParseError(TEXT("Unexpected comma"), ParseResult.Location); 
 				}
 				FunctionStack.Top().CountedCommas++;
 			}
 			else if (GetOperatorTokenInfo(Op).Associativity == EAssociativity::Left)
 			{
 				// Peel off all operators that have the same precedence or higher.
-				while (!OperatorStack.IsEmpty() && GetOperatorTokenInfo(OperatorStack.Last()).Precedence >= TokenInfo.Precedence)
+				while (!OperatorStack.IsEmpty() && GetOperatorTokenInfo(OperatorStack.Last().Get<0>()).Precedence >= TokenInfo.Precedence)
 				{
 					PushOperatorTokenToExpression(OperatorStack.Pop(false));
 				}
-				OperatorStack.Push(Op);
+				OperatorStack.Push({Op, ParseResult.Location});
 			}
 			else if (GetOperatorTokenInfo(Op).Associativity == EAssociativity::Right)
 			{
 				// Peel off all operators that have higher precedence.
-				while (!OperatorStack.IsEmpty() && GetOperatorTokenInfo(OperatorStack.Last()).Precedence > TokenInfo.Precedence)
+				while (!OperatorStack.IsEmpty() && GetOperatorTokenInfo(OperatorStack.Last().Get<0>()).Precedence > TokenInfo.Precedence)
 				{
 					PushOperatorTokenToExpression(OperatorStack.Pop(false));
 				}
-				OperatorStack.Push(Op);
+				OperatorStack.Push({Op, ParseResult.Location});
 			}
 		}
 
 		LastToken = Token;
-		LastTokenParseRange = ParseRange;
 		ParseRange = ParseRange.TrimStart();
 	} while(!ParseRange.IsEmpty());
 
 	while (!OperatorStack.IsEmpty())
 	{
-		if (!GetOperatorTokenInfo(OperatorStack.Last()).Operator.IsSet())
+		if (!GetOperatorTokenInfo(OperatorStack.Last().Get<0>()).Operator.IsSet())
 		{
-			return ParseError(TEXT("Mismatched parentheses"), InExpression.GetData()); 
+			return ParseError(TEXT("Mismatched parentheses"), OperatorStack.Last().Get<1>()); 
 		}
 		PushOperatorTokenToExpression(OperatorStack.Pop(false));
 	}
 
+	// If a constant evaluator function is given, then check if the constants exist.  
+	if (InConstantEvaluator)
+	{
+		for (int32 Index = 0; Index < Expression.Num(); Index++)
+		{
+			if (const FName* ConstantName = Expression[Index].TryGet<FName>())
+			{
+				if (!InConstantEvaluator(*ConstantName).IsSet())
+				{
+					return ParseError(*FString::Printf(TEXT("Unknown identifier '%s'"), *ConstantName->ToString()), Locations[Index]); 
+				}
+			}
+		}
+	}
+	
 	FExpressionObject S;
 	S.Expression = Expression;
 	return TVariant<FExpressionObject, FParseError>(TInPlaceType<FExpressionObject>(), S);
@@ -762,7 +828,8 @@ TVariant<FExpressionObject, FParseError> FEngine::Parse(
 
 
 float FEngine::Execute(
-	const FExpressionObject& InExpressionObject
+	const FExpressionObject& InExpressionObject,
+	TFunctionRef<TOptional<float>(FName InConstantName)> InConstantEvaluator
 	) const
 {
 	// An empty expression object is valid.
@@ -781,10 +848,10 @@ float FEngine::Execute(
 		{
 			ValueStack.Push(*Value);
 		}
-		else if (const FName* Constant = Token.TryGet<FName>())
+		else if (const FName* ConstantName = Token.TryGet<FName>())
 		{
-			const float *ConstantValue = Constants.Find(*Constant);
-			ValueStack.Push(ConstantValue ? *ConstantValue : 0.0f);
+			TOptional<float> ConstantValue = InConstantEvaluator(*ConstantName);
+			ValueStack.Push(ConstantValue.IsSet() ? ConstantValue.GetValue() : 0.0f);
 		}
 		else if (const FExpressionObject::EOperator* Operator = Token.TryGet<FExpressionObject::EOperator>())
 		{
@@ -886,13 +953,14 @@ float FEngine::Execute(
 
 
 TOptional<float> FEngine::Evaluate(
-	FStringView InExpression
+	FStringView InExpression,
+	TFunctionRef<TOptional<float>(FName InConstantName)> InConstantEvaluator
 	) const
 {
 	TVariant<FExpressionObject, FParseError> Result = Parse(InExpression);
 	if (const FExpressionObject* ExpressionObject = Result.TryGet<FExpressionObject>())
 	{
-		return Execute(*ExpressionObject);
+		return Execute(*ExpressionObject, InConstantEvaluator);
 	}
 
 	return {};
@@ -900,10 +968,11 @@ TOptional<float> FEngine::Evaluate(
 
 
 TOptional<FParseError> FEngine::Verify(
-	FStringView InExpression
+	FStringView InExpression,
+	TFunction<TOptional<float>(FName InConstantName)> InConstantEvaluator 
 	) const
 {
-	TVariant<FExpressionObject, FParseError> Result = Parse(InExpression);
+	TVariant<FExpressionObject, FParseError> Result = Parse(InExpression, InConstantEvaluator);
 	if (const FParseError* Error = Result.TryGet<FParseError>())
 	{
 		return *Error;
diff --git a/Engine/Plugins/Experimental/Animation/CurveExpression/Source/Runtime/Public/AnimNode_RemapCurves.h b/Engine/Plugins/Experimental/Animation/CurveExpression/Source/Runtime/Public/AnimNode_RemapCurves.h
new file mode 100644
--- /dev/null
+++ b/Engine/Plugins/Experimental/Animation/CurveExpression/Source/Runtime/Public/AnimNode_RemapCurves.h
@@ -0,0 +1,33 @@
+﻿// Copyright Epic Games, Inc. All Rights Reserved.
+
+#pragma once
+#include "AnimNode_RemapCurvesBase.h"
+
+#include "AnimNode_RemapCurves.generated.h"
+
+
+USTRUCT(BlueprintInternalUseOnly)
+struct CURVEEXPRESSION_API FAnimNode_RemapCurves :
+	public FAnimNode_RemapCurvesBase
+{
+	GENERATED_BODY()
+
+	// FAnimNode_Base interface
+	virtual void Initialize_AnyThread(const FAnimationInitializeContext& Context) override;
+	virtual void Evaluate_AnyThread(FPoseContext& Output) override;
+	
+	bool Serialize(FArchive& Ar);
+
+private:
+	// Used only by Initialize_AnyThread / Evaluate_AnyThread
+	TMap<FName, SmartName::UID_Type> CurveNameToUIDMap;
+	TArray<TTuple<SmartName::UID_Type, float>> CurveEvalResult;
+};
+
+template<> struct TStructOpsTypeTraits<FAnimNode_RemapCurves> : public TStructOpsTypeTraitsBase2<FAnimNode_RemapCurves>
+{
+	enum 
+	{ 
+		WithSerializer = true
+	};
+};
diff --git a/Engine/Plugins/Experimental/Animation/CurveExpression/Source/Runtime/Public/AnimNode_RemapCurvesFromMesh.h b/Engine/Plugins/Experimental/Animation/CurveExpression/Source/Runtime/Public/AnimNode_RemapCurvesFromMesh.h
--- a/Engine/Plugins/Experimental/Animation/CurveExpression/Source/Runtime/Public/AnimNode_RemapCurvesFromMesh.h
+++ b/Engine/Plugins/Experimental/Animation/CurveExpression/Source/Runtime/Public/AnimNode_RemapCurvesFromMesh.h
@@ -2,64 +2,51 @@
 
 #pragma once
 
-#include "ExpressionEvaluator.h"
-#include "Animation/AnimNodeBase.h"
+#include "AnimNode_RemapCurvesBase.h"
 
 #include "AnimNode_RemapCurvesFromMesh.generated.h"
 
+
 USTRUCT(BlueprintInternalUseOnly)
 struct CURVEEXPRESSION_API FAnimNode_RemapCurvesFromMesh :
-	public FAnimNode_Base
+	public FAnimNode_RemapCurvesBase
 {
 	GENERATED_BODY()
 
-	UPROPERTY(EditAnywhere, EditFixedSize, BlueprintReadWrite, Category = Links)
-	FPoseLink SourcePose;
-	
 	/** This is used by default if it's valid */
-	UPROPERTY(BlueprintReadWrite, transient, Category=Copy, meta=(PinShownByDefault))
+	UPROPERTY(BlueprintReadWrite, transient, Category=Copy, meta=(PinShownByDefault, DisplayAfter="SourcePose"))
 	TWeakObjectPtr<USkeletalMeshComponent> SourceMeshComponent;
 
 	/** If SourceMeshComponent is not valid, and if this is true, it will look for attached parent as a source */
 	UPROPERTY(EditAnywhere, BlueprintReadWrite, Category=Copy, meta=(NeverAsPin))
 	bool bUseAttachedParent = false;
 	
-	UPROPERTY(EditAnywhere, BlueprintReadWrite, editfixedsize, Category = Expressions, meta = (PinHiddenByDefault))
-	TMap<FName, FString> CurveExpressions;
-
-	/** The expressions given are immutable and will not change during runtime. Improves performance. */ 
-	UPROPERTY(EditAnywhere, BlueprintReadWrite, Category=Expressions, meta=(NeverAsPin))
-	bool bExpressionsImmutable = true;
-
-	// Call to verify the expressions. Write out any warnings to the InReportingFunc function.
-	void VerifyExpressions(
-		TFunction<void(FString)> InReportingFunc
-		) const;
-	
-	// Returns true if we're able to verify expressions at all. I.e. if the expression engine has not been initialized.
-	bool CanVerifyExpressions() const;
-
-	
 	// FAnimNode_Base interface
-	virtual void Initialize_AnyThread(const FAnimationInitializeContext& Context) override;
-	virtual void CacheBones_AnyThread(const FAnimationCacheBonesContext& Context) override;
-	virtual void Update_AnyThread(const FAnimationUpdateContext& Context) override;
+	virtual void PreUpdate(const UAnimInstance* InAnimInstance) override;
 	virtual void Evaluate_AnyThread(FPoseContext& Output) override;
 	virtual void GatherDebugData(FNodeDebugData& DebugData) override;
-	virtual bool HasPreUpdate() const override { return true; }
-	virtual void PreUpdate(const UAnimInstance* InAnimInstance) override;
+
+	bool Serialize(FArchive& Ar);
 	
 private:
 	void ReinitializeMeshComponent(USkeletalMeshComponent* InNewSkeletalMeshComponent, USkeletalMeshComponent* InTargetMeshComponent);
 	void RefreshMeshComponent(USkeletalMeshComponent* InTargetMeshComponent);
 
 	// this is source mesh references, so that we could compare and see if it has changed
-	TWeakObjectPtr<USkeletalMeshComponent>	CurrentlyUsedSourceMeshComponent;
-	TWeakObjectPtr<USkeletalMesh>			CurrentlyUsedSourceMesh;
+	TWeakObjectPtr<const USkeletalMeshComponent> CurrentlyUsedSourceMeshComponent;
+	TWeakObjectPtr<const USkeletalMesh> CurrentlyUsedSourceMesh;
 
 	// target mesh 
-	TWeakObjectPtr<USkeletalMesh>			CurrentlyUsedTargetMesh;
+	TWeakObjectPtr<const USkeletalMesh> CurrentlyUsedTargetMesh;
 
-	TOptional<CurveExpression::Evaluator::FEngine> ExpressionEngine;
-	TMap<FName, CurveExpression::Evaluator::FExpressionObject> CachedExpressions;
+	// Transient data.
+	TMap<FName, float> SourceCurveValues;
+};
+
+template<> struct TStructOpsTypeTraits<FAnimNode_RemapCurvesFromMesh> : public TStructOpsTypeTraitsBase2<FAnimNode_RemapCurvesFromMesh>
+{
+	enum 
+	{ 
+		WithSerializer = true
+	};
 };
diff --git a/Engine/Plugins/Experimental/Animation/CurveExpression/Source/Runtime/Public/CurveExpressionsDataAsset.h b/Engine/Plugins/Experimental/Animation/CurveExpression/Source/Runtime/Public/CurveExpressionsDataAsset.h
new file mode 100644
--- /dev/null
+++ b/Engine/Plugins/Experimental/Animation/CurveExpression/Source/Runtime/Public/CurveExpressionsDataAsset.h
@@ -0,0 +1,84 @@
+﻿// Copyright Epic Games, Inc. All Rights Reserved.
+
+#pragma once
+
+#include "ExpressionEvaluator.h"
+
+#include "Engine/DataAsset.h"
+
+#include "CurveExpressionsDataAsset.generated.h"
+
+
+struct FCurveExpressionAssignment
+{
+	int32 LineIndex = INDEX_NONE;
+	FName TargetName;
+	
+	FStringView Expression;
+};
+
+
+struct FCurveExpressionParsedAssignment
+{
+	int32 LineIndex = INDEX_NONE;
+	FName TargetName;
+		
+	TVariant<CurveExpression::Evaluator::FExpressionObject, CurveExpression::Evaluator::FParseError> Result; 
+};
+
+
+USTRUCT(BlueprintType)
+struct FCurveExpressionList
+{
+	GENERATED_BODY()
+
+	UPROPERTY(EditAnywhere, Category="Curve Expressions")
+	FString AssignmentExpressions;
+
+	TArray<FCurveExpressionAssignment> GetAssignments() const;
+	
+	/** Return compiled expression results */
+	TArray<FCurveExpressionParsedAssignment> GetParsedAssignments() const;
+};
+
+
+UCLASS(BlueprintType)
+class CURVEEXPRESSION_API UCurveExpressionsDataAsset :
+	public UDataAsset
+{
+	GENERATED_BODY()
+public:
+
+#if WITH_EDITORONLY_DATA
+	UPROPERTY(EditAnywhere, Category="Expressions")
+	FCurveExpressionList Expressions;
+#endif
+
+	/** Returns the list of named constants used by ExpressionList */
+	const TArray<FName>& GetCompiledExpressionConstants() const;
+
+	/** Returns the list of compiled expressions generated from ExpressionList */
+	const TMap<FName, CurveExpression::Evaluator::FExpressionObject>& GetCompiledExpressionMap() const;
+
+	// UObject overrides
+	void Serialize(FArchive& Ar) override;
+#if WITH_EDITOR
+	void PostEditChangeProperty(FPropertyChangedEvent& PropertyChangedEvent) override;
+#endif
+	
+private:
+#if WITH_EDITOR
+	void CompileExpressions();
+#endif
+
+	void SynchronizeThreadLocalData() const;
+	
+	mutable FCriticalSection ExpressionMapWriteLock;
+	
+	std::atomic<uint32> SerialNumber{0};
+	
+	UPROPERTY()
+	TArray<FName> NamedConstants;
+	
+	TMap<FName, CurveExpression::Evaluator::FExpressionObject> ExpressionMap;
+};
diff --git a/Engine/Plugins/Experimental/Animation/CurveExpression/Source/Runtime/Public/ExpressionEvaluator.h b/Engine/Plugins/Experimental/Animation/CurveExpression/Source/Runtime/Public/ExpressionEvaluator.h
--- a/Engine/Plugins/Experimental/Animation/CurveExpression/Source/Runtime/Public/ExpressionEvaluator.h
+++ b/Engine/Plugins/Experimental/Animation/CurveExpression/Source/Runtime/Public/ExpressionEvaluator.h
@@ -4,7 +4,6 @@
 
 #include "Containers/Map.h"
 #include "Misc/TVariant.h"
-#include "Misc/TVariantMeta.h"
 #include "UObject/NameTypes.h"
 
 template <typename OptionalType> struct TOptional;
@@ -12,22 +11,55 @@ template <typename OptionalType> struct TOptional;
 namespace CurveExpression::Evaluator
 {
 
+struct FParseLocation
+{
+	FParseLocation() = default;
+	FParseLocation(int32 InStart, int32 InEndLocation) :
+		Start(InStart), 
+		End(InEndLocation)
+	{}
+	FParseLocation(FStringView InParseRange, const TCHAR* InParseStart) :
+		Start(static_cast<int32>(InParseRange.begin() - InParseStart)),
+		End(static_cast<int32>(InParseRange.end() - InParseStart))
+	{}
+
+	/** Merge two parse locations and return a new one */
+	FParseLocation Join(const FParseLocation& InLocation) const
+	{
+		return FParseLocation(FMath::Min(Start, InLocation.Start), FMath::Max(End, InLocation.End));
+	}
+	
+	const int32 Start = INDEX_NONE;
+	const int32 End = INDEX_NONE;
+};
+
 struct FParseError
 {
 	FParseError() = default;
-	FParseError(const int32 InColumn, FString&& InMessage) : Column(InColumn), Message(MoveTemp(InMessage)) {}
-
-	/** The column at which the error occurs. Zero-based. */
-	const int32 Column = 0;
+	FParseError(FString&& InMessage, FParseLocation InLocation) :
+		Message(MoveTemp(InMessage)),
+		Location(InLocation)
+	{}
 
 	/** The human-readable parsing error */
-	const FString Message = FString();
+	const FString Message{};
+	
+	/** The parse location where the error took place. */
+	const FParseLocation Location{};
 };
 
 struct FExpressionObject
 {
 	FString ToString() const;
+
+	/** Returns all named constants used by this expression */
+	TSet<FName> GetUsedConstants() const;
+
+	void Load(FArchive& Ar);
+	void Save(FArchive& Ar) const;
+	
 private:
+	
 	enum class EOperator : int32
 	{
 		Negate,				// Negation operator.
@@ -74,74 +106,40 @@ private:
  *	rather than:
  *	   foo (1) * 2.0
  */
-enum class EParseFlags
-{
-	Default				= 0,		/** Default parsing behavior */
-	ValidateConstants	= 1 << 0,	/** Validate that the referred-to constant exists */
-};
-ENUM_CLASS_FLAGS(EParseFlags)
-
 
-/** Dictate how constants are updated. */
-enum class EUpdateConstantsMethod
-{
-	IgnoreMissing,					/** Existing constant values that are not updated are left unchanged. */
-	ZeroMissing						/** Existing constant values that are not updated are zeroed out. */
-};
 
 class CURVEEXPRESSION_API FEngine
 {
 public:
-	explicit FEngine(EParseFlags InParseFlags = EParseFlags::Default) :
-		ParseFlags(InParseFlags)
-	{ }
-
-	/** Initialize the engine with constants, making a copy of the value constant map. */
-	explicit FEngine(
-		const TMap<FName, float>& InConstants,
-		EParseFlags InParseFlags = EParseFlags::Default
-		);
-	
-	/** Initialize the engine with constants, taking ownership of the value constant map. */
-	explicit FEngine(
-		TMap<FName, float>&& InConstants,
-		EParseFlags InParseFlags = EParseFlags::Default
-		);
-
-	/** Update existing constant values. No new constants are added. */
-	void UpdateConstantValues(
-		const TMap<FName, float>& InConstants,
-		EUpdateConstantsMethod InUpdateMethod = EUpdateConstantsMethod::ZeroMissing
-		);
-
-	/** Return existing constants and their values. */
-	const TMap<FName, float>& GetConstantValues() const
-	{
-		return Constants;
-	}
-
 	/** Parse an expression and either receive an error or a expression execution object that
 	 *  matches this evaluator. */
 	TVariant<FExpressionObject, FParseError> Parse(
-		FStringView InExpression
+		FStringView InExpression,
+		TFunction<TOptional<float>(FName InConstantName)> InConstantEvaluator = {} 
 		) const;
 
-	/** Executes the given expression object */
+	/** Executes the given expression object.
+	 *  \param InExpressionObject The parsed expression object to evaluate.
+	 *  \param InConstantEvaluator A function that returns the value of a given constant, if it's known, or none otherwise.
+	 */
 	float Execute(
-		const FExpressionObject& InExpressionObject
+		const FExpressionObject& InExpressionObject,
+		TFunctionRef<TOptional<float>(FName InConstantName)> InConstantEvaluator
 		) const;
 	
 	/** Evaluate an expression directly and returns the result. If the expression is incorrectly
 	 *  formed, then the optional return value holds nothing.
 	 */ 
 	TOptional<float> Evaluate(
-		FStringView InExpression
+		FStringView InExpression,
+		TFunctionRef<TOptional<float>(FName InConstantName)> InConstantEvaluator
 		) const;
 
 	/** Verify an expression. Does not check for variable validity, only that the expression
 	 *  can be evaluated */
 	TOptional<FParseError> Verify(
-		FStringView InExpression
+		FStringView InExpression,
+		TFunction<TOptional<float>(FName InConstantName)> InConstantEvaluator = {} 
 		) const;
 	
 private:
@@ -162,11 +160,42 @@ private:
 		Max					// Must be last
 	};
 	using FToken = TVariant<EOperatorToken, /* Identifier */ FName, /* Value */ float>;
+
+	struct FTokenParseResult
+	{
+		FTokenParseResult() = default;
+		
+		FTokenParseResult(FToken&& InToken, FStringView InTokenRange, const TCHAR* InExpression) :
+			Result(TInPlaceType<FToken>(), MoveTemp(InToken)),
+			Location(InTokenRange, InExpression)
+		{}
+		
+		FTokenParseResult(FString&& InError, FStringView InTokenRange, const TCHAR* InExpression) :
+			Result(TInPlaceType<FString>(), MoveTemp(InError)),
+			Location(InTokenRange, InExpression)
+		{}
+		
+		bool IsError() const { return Result.IsType<FString>(); }
+		bool IsToken() const { return Result.IsType<FToken>(); }
+		
+		const FToken& GetToken() const
+		{
+			return Result.Get<FToken>();
+		}
+		
+		FParseError GetParseError()
+		{
+			return FParseError(MoveTemp(Result.Get<FString>()), Location);
+		}
+		
+		TVariant<FString, FToken> Result;
+		FParseLocation Location{};
+	};
 	
 	/** Parse a floating point value and return the result (or error). Same semantics
 	 *  apply as with ParseToken.
 	 */
-	static TVariant<FToken, FParseError> ParseFloat(
+	static FTokenParseResult ParseFloat(
 		FStringView& InOutParseRange,
 		FStringView InExpression
 		);
@@ -174,39 +203,26 @@ private:
 	/** Parse a constant reference and return the result (or error). Same semantics
 	 *  apply as with ParseToken.
 	 */
-	TVariant<FToken, FParseError> ParseIdentifier(
+	static FTokenParseResult ParseIdentifier(
 		FStringView& InOutParseRange,
 		FStringView InExpression
-		) const;
+		);
 	
 	/** Parse a single token from the beginning of the range, returning either the token
 	 *  or an error, if one was encountered.
 	 *  If the token was successfully parsed, then the given string view range gets updated
 	 *  with the new starting point at the end of the token just parsed
 	 */
-	TVariant<FToken, FParseError> ParseToken(
+	static FTokenParseResult ParseToken(
 		FStringView& InOutParseRange,
 		FStringView InExpression
-		) const;
-
-	/** Helper function to construct a token parse error result */
-	static TVariant<FToken, FParseError> ParseTokenError(
-		FString&& InError,
-		const TCHAR* InCurrentPos,
-		FStringView InExpression
 		);
-	
+
 	/** Helper function to get a text version of a token. */
 	static FString TokenToString(
 		const FToken& InToken
 		);
 	
-	/** The list of constants, stored using the FName key since the expression object refers
-	 *  to it directly */
-	// TBD: Add support for variables by allowing TFunction{Ref} objects along with direct
-	// constants.
-	TMap<FName, float> Constants;
-
 	/** For now we just have built-in functions */
 	friend struct FInitializeBuiltinFunctions;
 	using FunctionType = TFunction<float(TArrayView<const float>)>;
@@ -217,9 +233,20 @@ private:
 	};
 	static TMap<FName, int32> FunctionNameIndex;
 	static TArray<FFunctionInfo> Functions;
-
-	/** The parse flags for this engine. Only set once */
-	EParseFlags ParseFlags = EParseFlags::Default;
 };
 
 }
+
+
+static FArchive& operator<<(FArchive& Ar, CurveExpression::Evaluator::FExpressionObject& InObject)
+{
+	if (Ar.IsLoading())
+	{
+		InObject.Load(Ar);
+	}
+	else
+	{
+		InObject.Save(Ar);
+	}
+	return Ar;
+}
