[Backout] - CL24806012
Original CL Desc
-----------------------------------------------------------------
Curve Expression: Curve Expression List Data Asset

Added support for a curve expression data asset that can be swapped out at run-time. The asset contains compiled expressions so swapping in and out won't incur any runtime expression parsing costs.
Also added the ability to set expressions directly on the node.
Added a new node, "Remap Curves" which allows remapping curves within the same pose context.
Fixed the expression evaluator in preparation for new curve work.

diff --git a/Engine/Plugins/Experimental/Animation/CurveExpression/Source/Editor/CurveExpressionEditor.Build.cs b/Engine/Plugins/Experimental/Animation/CurveExpression/Source/Editor/CurveExpressionEditor.Build.cs
--- a/Engine/Plugins/Experimental/Animation/CurveExpression/Source/Editor/CurveExpressionEditor.Build.cs
+++ b/Engine/Plugins/Experimental/Animation/CurveExpression/Source/Editor/CurveExpressionEditor.Build.cs
@@ -1,25 +1,25 @@
 // Copyright Epic Games, Inc. All Rights Reserved.
 
-using EpicGames.Core;
 using UnrealBuildTool;
 
 public class CurveExpressionEditor : ModuleRules
 {
 	public CurveExpressionEditor(ReadOnlyTargetRules Target) : base(Target)
 	{
-		PrivateDependencyModuleNames.AddAll(
-			"AnimGraph",
-			"AssetDefinition",
-			"BlueprintGraph",
-			"Core",
-			"CoreUObject",
-			"CurveExpression",
-			"Engine", 
-			"Kismet",
-			"KismetCompiler",
-			"Slate",
-			"SlateCore",
-			"UnrealEd"
+		PrivateDependencyModuleNames.AddRange(
+			new string[]
+			{
+				"AnimGraph",
+				"BlueprintGraph",
+				"Core",
+				"CoreUObject",
+				"CurveExpression",
+				"Engine",
+				"KismetCompiler",
+				"Slate",
+				"SlateCore",
+				"UnrealEd",
+			}
 		);
 	}
 }
diff --git a/Engine/Plugins/Experimental/Animation/CurveExpression/Source/Editor/Private/AnimGraphNode_RemapCurves.cpp b/Engine/Plugins/Experimental/Animation/CurveExpression/Source/Editor/Private/AnimGraphNode_RemapCurves.cpp
deleted file mode 100644
diff --git a/Engine/Plugins/Experimental/Animation/CurveExpression/Source/Editor/Private/AnimGraphNode_RemapCurvesFromMesh.cpp b/Engine/Plugins/Experimental/Animation/CurveExpression/Source/Editor/Private/AnimGraphNode_RemapCurvesFromMesh.cpp
--- a/Engine/Plugins/Experimental/Animation/CurveExpression/Source/Editor/Private/AnimGraphNode_RemapCurvesFromMesh.cpp
+++ b/Engine/Plugins/Experimental/Animation/CurveExpression/Source/Editor/Private/AnimGraphNode_RemapCurvesFromMesh.cpp
@@ -7,12 +7,12 @@
 #include "Animation/AnimAttributes.h"
 #include "Animation/AnimInstance.h"
 #include "Components/SkeletalMeshComponent.h"
-#include "Kismet2/BlueprintEditorUtils.h"
 
 #include UE_INLINE_GENERATED_CPP_BY_NAME(AnimGraphNode_RemapCurvesFromMesh)
 
 #define LOCTEXT_NAMESPACE "AnimGraphNode_RemapCurvesFromMesh"
 
+
 bool UAnimGraphNode_RemapCurvesFromMesh::CanVerifyExpressions() const
 {
 	if (const FAnimNode_RemapCurvesFromMesh* DebuggedNode = GetDebuggedNode())
@@ -27,19 +27,8 @@ void UAnimGraphNode_RemapCurvesFromMesh::VerifyExpressions()
 {
 	if (const FAnimNode_RemapCurvesFromMesh* DebuggedNode = GetDebuggedNode())
 	{
-		const USkeletalMeshComponent* TargetComponent = GetDebuggedComponent();
-		const USkeletalMeshComponent* SourceComponent = nullptr;
-		if (Node.SourceMeshComponent.IsValid())
-		{
-			SourceComponent = Node.SourceMeshComponent.Get();
-		}
-		else if (Node.bUseAttachedParent && TargetComponent)
-		{
-			SourceComponent = Cast<USkeletalMeshComponent>(TargetComponent->GetAttachParent());
-		}
-		
 		TArray<FString> Results;
-		DebuggedNode->VerifyExpressions(TargetComponent, SourceComponent, [&Results](const FString& InMessage){
+		DebuggedNode->VerifyExpressions([&Results](const FString& InMessage){
 			Results.Add(InMessage);
 		});
 
@@ -84,57 +73,7 @@ void UAnimGraphNode_RemapCurvesFromMesh::GetOutputLinkAttributes(FNodeAttributeA
 }
 
 
-void UAnimGraphNode_RemapCurvesFromMesh::CustomizePinData(UEdGraphPin* Pin, FName SourcePropertyName, int32 ArrayIndex) const
-{
-	Super::CustomizePinData(Pin, SourcePropertyName, ArrayIndex);
-
-	if (Pin->PinName == GET_MEMBER_NAME_STRING_CHECKED(FAnimNode_RemapCurvesBase, CurveExpressionsDataAsset))
-	{
-		Pin->bHidden = (Node.ExpressionSource != ERemapCurvesExpressionSource::DataAsset);
-	}
-	if (Pin->PinName == GET_MEMBER_NAME_STRING_CHECKED(FAnimNode_RemapCurvesBase, CurveExpressions))
-	{
-		Pin->bHidden = (Node.ExpressionSource != ERemapCurvesExpressionSource::ExpressionMap);
-	}
-}
-
-
-void UAnimGraphNode_RemapCurvesFromMesh::PostEditChangeProperty(FPropertyChangedEvent& PropertyChangedEvent)
-{
-	Super::PostEditChangeProperty(PropertyChangedEvent);
-
-	const FName PropertyName = PropertyChangedEvent.GetPropertyName();
-	if (PropertyName == GET_MEMBER_NAME_CHECKED(FCurveExpressionList, AssignmentExpressions) || 
-	 	PropertyName == GET_MEMBER_NAME_CHECKED(FAnimNode_RemapCurvesFromMesh, ExpressionSource))
-	{
-		Node.ParseAndCacheExpressions();
-	}
-	if (PropertyName == GET_MEMBER_NAME_CHECKED(FAnimNode_RemapCurvesBase, ExpressionSource))
-	{
-		{
-			FScopedTransaction Transaction(LOCTEXT("ChangeExpressionSource", "Change Expression Source"));
-			Modify();
-
-			for (UEdGraphPin* Pin: Pins)
-			{
-				if (Pin->PinName == GET_MEMBER_NAME_STRING_CHECKED(FAnimNode_RemapCurvesBase, CurveExpressionsDataAsset) &&
-					Node.ExpressionSource != ERemapCurvesExpressionSource::DataAsset)
-				{
-					Pin->BreakAllPinLinks();
-				}
-				if (Pin->PinName == GET_MEMBER_NAME_STRING_CHECKED(FAnimNode_RemapCurvesBase, CurveExpressions) &&
-					Node.ExpressionSource != ERemapCurvesExpressionSource::ExpressionMap)
-				{
-					Pin->BreakAllPinLinks();
-				}
-			}
-		}
-		ReconstructNode();
-		FBlueprintEditorUtils::MarkBlueprintAsStructurallyModified(GetBlueprint());
-	}
-}
-
-USkeletalMeshComponent* UAnimGraphNode_RemapCurvesFromMesh::GetDebuggedComponent() const
+FAnimNode_RemapCurvesFromMesh* UAnimGraphNode_RemapCurvesFromMesh::GetDebuggedNode() const
 {
 	if (const UObject* ObjectBeingDebugged = GetAnimBlueprint()->GetObjectBeingDebugged())
 	{
@@ -143,20 +82,10 @@ USkeletalMeshComponent* UAnimGraphNode_RemapCurvesFromMesh::GetDebuggedComponent
 			USkeletalMeshComponent* Component = InstanceBeingDebugged->GetSkelMeshComponent();
 			if (Component != nullptr && Component->GetAnimInstance() != nullptr)
 			{
-				return Component;
+				return static_cast<FAnimNode_RemapCurvesFromMesh*>(FindDebugAnimNode(Component));
 			}
 		}
 	}
-	return nullptr;
-}
-
-
-FAnimNode_RemapCurvesFromMesh* UAnimGraphNode_RemapCurvesFromMesh::GetDebuggedNode() const
-{
-	if (USkeletalMeshComponent* Component = GetDebuggedComponent(); Component != nullptr)
-	{
-		return static_cast<FAnimNode_RemapCurvesFromMesh*>(FindDebugAnimNode(Component));
-	}
 
 	return nullptr;
 }
diff --git a/Engine/Plugins/Experimental/Animation/CurveExpression/Source/Editor/Private/AssetDefinition_CurveExpressionsDataAsset.cpp b/Engine/Plugins/Experimental/Animation/CurveExpression/Source/Editor/Private/AssetDefinition_CurveExpressionsDataAsset.cpp
deleted file mode 100644
diff --git a/Engine/Plugins/Experimental/Animation/CurveExpression/Source/Editor/Private/AssetDefinition_CurveExpressionsDataAsset.h b/Engine/Plugins/Experimental/Animation/CurveExpression/Source/Editor/Private/AssetDefinition_CurveExpressionsDataAsset.h
deleted file mode 100644
diff --git a/Engine/Plugins/Experimental/Animation/CurveExpression/Source/Editor/Private/CurveExpressionDetailsCustomization.cpp b/Engine/Plugins/Experimental/Animation/CurveExpression/Source/Editor/Private/CurveExpressionDetailsCustomization.cpp
--- a/Engine/Plugins/Experimental/Animation/CurveExpression/Source/Editor/Private/CurveExpressionDetailsCustomization.cpp
+++ b/Engine/Plugins/Experimental/Animation/CurveExpression/Source/Editor/Private/CurveExpressionDetailsCustomization.cpp
@@ -9,32 +9,33 @@
 #include "DetailLayoutBuilder.h"
 #include "DetailWidgetRow.h"
 #include "Widgets/Input/SButton.h"
-#include "Widgets/Input/SMultiLineEditableTextBox.h"
 #include "Widgets/Layout/SGridPanel.h"
 #include "Widgets/Layout/SScrollBox.h"
+#include "Widgets/Text/SMultiLineEditableText.h"
 
 
 #define LOCTEXT_NAMESPACE "CurveExpressionDetailsCustomization"
 
-TSharedRef<IDetailCustomization> FAnimGraphNode_RemapCurvesDebuggingCustomization::MakeInstance()
+TSharedRef<IDetailCustomization> FAnimGraphNode_RemapCurvesFromMeshCustomization::MakeInstance()
 {
-	return MakeShared<FAnimGraphNode_RemapCurvesDebuggingCustomization>();
+	return MakeShared<FAnimGraphNode_RemapCurvesFromMeshCustomization>();
 }
 
 
-void FAnimGraphNode_RemapCurvesDebuggingCustomization::CustomizeDetails(
+void FAnimGraphNode_RemapCurvesFromMeshCustomization::CustomizeDetails(
 	IDetailLayoutBuilder& InDetailBuilder
 	)
 {
 	TArray<TWeakObjectPtr<UObject>> ObjectsBeingCustomized;
 	InDetailBuilder.GetObjectsBeingCustomized(ObjectsBeingCustomized);
 
-	TArray<TWeakObjectPtr<UObject>> RemapNodes;
+	TArray<TWeakObjectPtr<UAnimGraphNode_RemapCurvesFromMesh>> RemapNodes;
 	for (TWeakObjectPtr<UObject> Object: ObjectsBeingCustomized)
 	{
-		if (Cast<IRemapCurvesDebuggingProvider>(Object.Get()))
+		TWeakObjectPtr<UAnimGraphNode_RemapCurvesFromMesh> RemapNode = Cast<UAnimGraphNode_RemapCurvesFromMesh>(Object.Get());
+		if (RemapNode.IsValid())
 		{
-			RemapNodes.Add(Object);
+			RemapNodes.Add(RemapNode);
 		}
 	}
 
@@ -52,9 +53,9 @@ void FAnimGraphNode_RemapCurvesDebuggingCustomization::CustomizeDetails(
 				.ToolTipText(LOCTEXT("VerifyExpressions_Tooltip", "Verify all the expressions as applied to the currently debugged object. If debugging is not active, no verification is done. Verification output is sent to the output log."))
 				.OnClicked_Lambda([this, RemapNodes]()
 				{
-					for (TWeakObjectPtr<UObject> Object: RemapNodes)
+					for (TWeakObjectPtr<UAnimGraphNode_RemapCurvesFromMesh> RemapNodeWeak: RemapNodes)
 					{
-						if (IRemapCurvesDebuggingProvider* RemapNode = Cast<IRemapCurvesDebuggingProvider>(Object.Get()))
+						if (UAnimGraphNode_RemapCurvesFromMesh* RemapNode = RemapNodeWeak.Get())
 						{
 							RemapNode->VerifyExpressions();
 						}
@@ -63,9 +64,9 @@ void FAnimGraphNode_RemapCurvesDebuggingCustomization::CustomizeDetails(
 				})
 				.IsEnabled_Lambda([this, RemapNodes]()
 				{
-					for (TWeakObjectPtr<UObject> Object: RemapNodes)
+					for (TWeakObjectPtr<UAnimGraphNode_RemapCurvesFromMesh> RemapNodeWeak: RemapNodes)
 					{
-						if (const IRemapCurvesDebuggingProvider* RemapNode = Cast<IRemapCurvesDebuggingProvider>(Object.Get()))
+						if (UAnimGraphNode_RemapCurvesFromMesh* RemapNode = RemapNodeWeak.Get())
 						{
 							return RemapNode->CanVerifyExpressions();
 						}
@@ -91,24 +92,18 @@ void FCurveExpressionListCustomization::CustomizeHeader(
 {
 	AssignmentExpressionsProperty = InPropertyHandle->GetChildHandle(GET_MEMBER_NAME_CHECKED(FCurveExpressionList, AssignmentExpressions));
 
-	const FScrollBarStyle& ScrollBarStyle = FCurveExpressionEditorStyle::Get().GetWidgetStyle<FScrollBarStyle>("TextEditor.ScrollBar"); 
-
 	HorizontalScrollbar =
 		SNew(SScrollBar)
-			.Style(&ScrollBarStyle)
 			.AlwaysShowScrollbar(true)
-			.Orientation(Orient_Horizontal)
-			.Thickness(FVector2D(6.0f, 6.0f));
+			.Orientation(Orient_Horizontal);
 
 	VerticalScrollbar =
 		SNew(SScrollBar)
-			.Style(&ScrollBarStyle)
 			.AlwaysShowScrollbar(true)
-			.Orientation(Orient_Vertical)
-			.Thickness(FVector2D(6.0f, 6.0f));
+			.Orientation(Orient_Vertical);
 
-	const FEditableTextBoxStyle &TextBoxStyle = FCurveExpressionEditorStyle::Get().GetWidgetStyle<FEditableTextBoxStyle>("TextEditor.EditableTextBox");
-	const FSlateFontInfo &Font = TextBoxStyle.TextStyle.Font;
+	const FTextBlockStyle &TextStyle = FCurveExpressionEditorStyle::Get().GetWidgetStyle<FTextBlockStyle>("TextEditor.NormalText");
+	const FSlateFontInfo &Font = TextStyle.Font;
 	
 	InHeaderRow
 	.WholeRowContent()
@@ -127,24 +122,44 @@ void FCurveExpressionListCustomization::CustomizeHeader(
 			.BorderImage(FCurveExpressionEditorStyle::Get().GetBrush("TextEditor.Border"))
 			.BorderBackgroundColor(FLinearColor::Black)
 			[
-				SAssignNew(TextEditor, SMultiLineEditableTextBox)
-				.Font(Font)
-				.Style(&FCoreStyle::Get().GetWidgetStyle<FEditableTextBoxStyle>("NormalEditableTextBox"))
-				.Text_Lambda([this]()
-				{
-					FString Text;
-					AssignmentExpressionsProperty->GetValue(Text);
-					return FText::FromString(Text);
-				})
-				.OnTextChanged_Lambda([this](const FText& InText)
-				{
-					AssignmentExpressionsProperty->SetValue(InText.ToString());
-				})
-				// By default, the Tab key gets routed to "next widget". We want to disable that behaviour.
-				.OnIsTypedCharValid_Lambda([](const TCHAR InChar) { return true; })
-				.AutoWrapText(false)
-				.HScrollBar(HorizontalScrollbar)
-				.VScrollBar(VerticalScrollbar)
+				SNew(SGridPanel)
+				.FillColumn(0, 1.0f)
+				.FillRow(0, 1.0f)
+				+SGridPanel::Slot(0, 0)
+				[
+					SNew(SScrollBox)
+					.Orientation(Orient_Vertical)
+					.ExternalScrollbar(VerticalScrollbar)
+					+ SScrollBox::Slot()
+					[
+						SAssignNew(TextEditor, SMultiLineEditableText)
+						.Font(Font)
+						.TextStyle(&TextStyle)
+						.Text_Lambda([this]()
+						{
+							FString Text;
+							AssignmentExpressionsProperty->GetValue(Text);
+							return FText::FromString(Text);
+						})
+						.OnTextChanged_Lambda([this](const FText& InText)
+						{
+							AssignmentExpressionsProperty->SetValue(InText.ToString());
+						})
+						// By default, the Tab key gets routed to "next widget". We want to disable that behaviour.
+						.OnIsTypedCharValid_Lambda([](const TCHAR InChar) { return true; })
+						.AutoWrapText(false)
+						.HScrollBar(HorizontalScrollbar)
+						.VScrollBar(VerticalScrollbar)
+					]
+				]
+				+SGridPanel::Slot(1, 0)
+				[
+					VerticalScrollbar.ToSharedRef()
+				]
+				+SGridPanel::Slot(0, 1)
+				[
+					HorizontalScrollbar.ToSharedRef()
+				]
 			]
 		]
 	];
diff --git a/Engine/Plugins/Experimental/Animation/CurveExpression/Source/Editor/Private/CurveExpressionDetailsCustomization.h b/Engine/Plugins/Experimental/Animation/CurveExpression/Source/Editor/Private/CurveExpressionDetailsCustomization.h
--- a/Engine/Plugins/Experimental/Animation/CurveExpression/Source/Editor/Private/CurveExpressionDetailsCustomization.h
+++ b/Engine/Plugins/Experimental/Animation/CurveExpression/Source/Editor/Private/CurveExpressionDetailsCustomization.h
@@ -6,11 +6,11 @@
 #include "IPropertyTypeCustomization.h"
 
 
-class SMultiLineEditableTextBox;
+class SMultiLineEditableText;
 class SScrollBar;
 
 
-class FAnimGraphNode_RemapCurvesDebuggingCustomization :
+class FAnimGraphNode_RemapCurvesFromMeshCustomization :
 	public IDetailCustomization
 {
 public:
@@ -44,6 +44,6 @@ private:
 	TSharedPtr<SScrollBar> HorizontalScrollbar;
 	TSharedPtr<SScrollBar> VerticalScrollbar;
 
-	TSharedPtr<SMultiLineEditableTextBox> TextEditor;	
+	TSharedPtr<SMultiLineEditableText> TextEditor;	
 };
 
diff --git a/Engine/Plugins/Experimental/Animation/CurveExpression/Source/Editor/Private/CurveExpressionEditorModule.cpp b/Engine/Plugins/Experimental/Animation/CurveExpression/Source/Editor/Private/CurveExpressionEditorModule.cpp
--- a/Engine/Plugins/Experimental/Animation/CurveExpression/Source/Editor/Private/CurveExpressionEditorModule.cpp
+++ b/Engine/Plugins/Experimental/Animation/CurveExpression/Source/Editor/Private/CurveExpressionEditorModule.cpp
@@ -1,6 +1,5 @@
 // Copyright Epic Games, Inc. All Rights Reserved.
 
-#include "AnimGraphNode_RemapCurves.h"
 #include "AnimGraphNode_RemapCurvesFromMesh.h"
 #include "CurveExpressionDetailsCustomization.h"
 #include "CurveExpressionEditorStyle.h"
@@ -41,14 +40,9 @@ public:
 			&FCurveExpressionListCustomization::MakeInstance
 			);
 			
-		RegisterClassCustomization(
-			UAnimGraphNode_RemapCurves::StaticClass()->GetFName(),
-			&FAnimGraphNode_RemapCurvesDebuggingCustomization::MakeInstance
-			);
-		
 		RegisterClassCustomization(
 			UAnimGraphNode_RemapCurvesFromMesh::StaticClass()->GetFName(),
-			&FAnimGraphNode_RemapCurvesDebuggingCustomization::MakeInstance
+			&FAnimGraphNode_RemapCurvesFromMeshCustomization::MakeInstance
 			);
 	}
 	virtual void ShutdownModule() override
diff --git a/Engine/Plugins/Experimental/Animation/CurveExpression/Source/Editor/Private/CurveExpressionEditorStyle.cpp b/Engine/Plugins/Experimental/Animation/CurveExpression/Source/Editor/Private/CurveExpressionEditorStyle.cpp
--- a/Engine/Plugins/Experimental/Animation/CurveExpression/Source/Editor/Private/CurveExpressionEditorStyle.cpp
+++ b/Engine/Plugins/Experimental/Animation/CurveExpression/Source/Editor/Private/CurveExpressionEditorStyle.cpp
@@ -54,18 +54,19 @@ FCurveExpressionEditorStyle::FCurveExpressionEditorStyle() :
 		const FTextBlockStyle ErrorText = FTextBlockStyle(NormalText)
 			.SetUnderlineBrush(IMAGE_BRUSH("Images/White", Icon8x8, FLinearColor::Red, ESlateBrushTileType::Both))
 			.SetColorAndOpacity(FLinearColor::Red);
+		// ;
 		
 		Set("TextEditor.NormalText", NormalText);
 		
 		Set("TextEditor.Border", new BOX_BRUSH("Images/TextEditorBorder", FMargin(4.0f/16.0f), FLinearColor(0.02f,0.02f,0.02f,1)));
 
-		const FEditableTextBoxStyle EditableTextBoxStyle = FEditableTextBoxStyle(FCoreStyle::Get().GetWidgetStyle<FEditableTextBoxStyle>("NormalEditableTextBox"))
-			.SetTextStyle(NormalText);
+		const FEditableTextBoxStyle EditableTextBoxStyle = FEditableTextBoxStyle()
+			.SetTextStyle(NormalText)
+			.SetBackgroundImageNormal( FSlateNoResource() )
+			.SetBackgroundImageHovered( FSlateNoResource() )
+			.SetBackgroundImageFocused( FSlateNoResource() )
+			.SetBackgroundImageReadOnly( FSlateNoResource() );
 		
 		Set("TextEditor.EditableTextBox", EditableTextBoxStyle);
-
-		const FScrollBarStyle ScrollBarStyle = FScrollBarStyle(FCoreStyle::Get().GetWidgetStyle<FScrollBarStyle>("ScrollBar"))
-			.SetThickness(6.0f);
-		Set("TextEditor.ScrollBar", ScrollBarStyle);
 	}
 }
diff --git a/Engine/Plugins/Experimental/Animation/CurveExpression/Source/Editor/Private/CurveExpressionsDataAssetFactory.cpp b/Engine/Plugins/Experimental/Animation/CurveExpression/Source/Editor/Private/CurveExpressionsDataAssetFactory.cpp
deleted file mode 100644
diff --git a/Engine/Plugins/Experimental/Animation/CurveExpression/Source/Editor/Private/CurveExpressionsDataAssetFactory.h b/Engine/Plugins/Experimental/Animation/CurveExpression/Source/Editor/Private/CurveExpressionsDataAssetFactory.h
deleted file mode 100644
diff --git a/Engine/Plugins/Experimental/Animation/CurveExpression/Source/Editor/Private/K2Node_MakeCurveExpressionMap.cpp b/Engine/Plugins/Experimental/Animation/CurveExpression/Source/Editor/Private/K2Node_MakeCurveExpressionMap.cpp
--- a/Engine/Plugins/Experimental/Animation/CurveExpression/Source/Editor/Private/K2Node_MakeCurveExpressionMap.cpp
+++ b/Engine/Plugins/Experimental/Animation/CurveExpression/Source/Editor/Private/K2Node_MakeCurveExpressionMap.cpp
@@ -196,7 +196,7 @@ void UK2Node_MakeCurveExpressionMap::ValidateNodeDuringCompilation(FCompilerResu
 				{
 					const int32 ExpressionColumn = SourceExpression.GetData() - InLine.GetData();
 					Args.Add(TEXT("Error"), FText::FromString(Error->Message));
-					Args.Add(TEXT("Column"), FText::AsNumber(Error->Location.Start + ExpressionColumn + 1));
+					Args.Add(TEXT("Column"), FText::AsNumber(Error->Column + ExpressionColumn + 1));
 				
 					MessageLog.Error(
 						*FText::Format(LOCTEXT("Warning_BadExpression", "@@ has an error on line {Line}, column {Column}. {Error}"), Args).ToString(), 
diff --git a/Engine/Plugins/Experimental/Animation/CurveExpression/Source/Editor/Public/AnimGraphNode_RemapCurves.h b/Engine/Plugins/Experimental/Animation/CurveExpression/Source/Editor/Public/AnimGraphNode_RemapCurves.h
deleted file mode 100644
diff --git a/Engine/Plugins/Experimental/Animation/CurveExpression/Source/Editor/Public/AnimGraphNode_RemapCurvesFromMesh.h b/Engine/Plugins/Experimental/Animation/CurveExpression/Source/Editor/Public/AnimGraphNode_RemapCurvesFromMesh.h
--- a/Engine/Plugins/Experimental/Animation/CurveExpression/Source/Editor/Public/AnimGraphNode_RemapCurvesFromMesh.h
+++ b/Engine/Plugins/Experimental/Animation/CurveExpression/Source/Editor/Public/AnimGraphNode_RemapCurvesFromMesh.h
@@ -5,7 +5,6 @@
 #include "AnimGraphNode_Base.h"
 
 #include "AnimNode_RemapCurvesFromMesh.h"
-#include "IRemapCurvesDebuggingProvider.h"
 
 #include "AnimGraphNode_RemapCurvesFromMesh.generated.h"
 
@@ -14,34 +13,29 @@ namespace ENodeTitleType { enum Type : int; }
 
 UCLASS(MinimalAPI)
 class UAnimGraphNode_RemapCurvesFromMesh :
-	public UAnimGraphNode_Base,
-	public IRemapCurvesDebuggingProvider
+	public UAnimGraphNode_Base
 {
 	GENERATED_BODY()
 
 public:
-	UPROPERTY(EditAnywhere, BlueprintReadWrite, Category=Settings)
+	UPROPERTY(EditAnywhere, Category=Settings)
 	FAnimNode_RemapCurvesFromMesh Node;
 
-	bool CanVerifyExpressions() const override;
-	void VerifyExpressions() override;
+	bool CanVerifyExpressions() const;
+	void VerifyExpressions();
 
 	// UEdGraphNode interface
-	FText GetNodeTitle(ENodeTitleType::Type TitleType) const override;
-	FText GetTooltipText() const override;
+	virtual FText GetNodeTitle(ENodeTitleType::Type TitleType) const override;
+	virtual FText GetTooltipText() const override;
 	FText GetMenuCategory() const override;
 	// End of UEdGraphNode interface
 
 	// UAnimGraphNode_Base interface
-	bool UsingCopyPoseFromMesh() const override { return true; }
-	void GetOutputLinkAttributes(FNodeAttributeArray& OutAttributes) const override;
-	void CustomizePinData(UEdGraphPin* Pin, FName SourcePropertyName, int32 ArrayIndex) const override;
+	virtual bool UsingCopyPoseFromMesh() const override { return true; }
+	virtual void GetOutputLinkAttributes(FNodeAttributeArray& OutAttributes) const override;
 	// End of UAnimGraphNode_Base interface
 	
-	// UObject override
-	void PostEditChangeProperty(FPropertyChangedEvent& PropertyChangedEvent) override; 	
-	
 private:
-	USkeletalMeshComponent* GetDebuggedComponent() const;
 	FAnimNode_RemapCurvesFromMesh* GetDebuggedNode() const;
+	
 };
diff --git a/Engine/Plugins/Experimental/Animation/CurveExpression/Source/Editor/Public/IRemapCurvesDebuggingProvider.h b/Engine/Plugins/Experimental/Animation/CurveExpression/Source/Editor/Public/IRemapCurvesDebuggingProvider.h
deleted file mode 100644
diff --git a/Engine/Plugins/Experimental/Animation/CurveExpression/Source/Editor/Public/K2Node_MakeCurveExpressionMap.h b/Engine/Plugins/Experimental/Animation/CurveExpression/Source/Editor/Public/K2Node_MakeCurveExpressionMap.h
--- a/Engine/Plugins/Experimental/Animation/CurveExpression/Source/Editor/Public/K2Node_MakeCurveExpressionMap.h
+++ b/Engine/Plugins/Experimental/Animation/CurveExpression/Source/Editor/Public/K2Node_MakeCurveExpressionMap.h
@@ -2,8 +2,6 @@
 
 #pragma once
 
-#include "CurveExpressionsDataAsset.h"
-
 #include "K2Node.h"
 
 #include "K2Node_MakeCurveExpressionMap.generated.h"
@@ -12,6 +10,14 @@ class FCompilerResultsLog;
 class FNodeHandlingFunctor;
 namespace ENodeTitleType { enum Type : int; }
 
+USTRUCT()
+struct FCurveExpressionList
+{
+	GENERATED_BODY()
+
+	UPROPERTY(EditAnywhere, Category="Expressions")
+	FString AssignmentExpressions;
+};
 
 UCLASS()
 class CURVEEXPRESSIONEDITOR_API UK2Node_MakeCurveExpressionMap :
diff --git a/Engine/Plugins/Experimental/Animation/CurveExpression/Source/Runtime/Private/AnimNode_RemapCurves.cpp b/Engine/Plugins/Experimental/Animation/CurveExpression/Source/Runtime/Private/AnimNode_RemapCurves.cpp
deleted file mode 100644
diff --git a/Engine/Plugins/Experimental/Animation/CurveExpression/Source/Runtime/Private/AnimNode_RemapCurvesFromMesh.cpp b/Engine/Plugins/Experimental/Animation/CurveExpression/Source/Runtime/Private/AnimNode_RemapCurvesFromMesh.cpp
--- a/Engine/Plugins/Experimental/Animation/CurveExpression/Source/Runtime/Private/AnimNode_RemapCurvesFromMesh.cpp
+++ b/Engine/Plugins/Experimental/Animation/CurveExpression/Source/Runtime/Private/AnimNode_RemapCurvesFromMesh.cpp
@@ -1,64 +1,117 @@
 // Copyright Epic Games, Inc. All Rights Reserved.
 
 #include "AnimNode_RemapCurvesFromMesh.h"
-
-#include "CurveExpressionCustomVersion.h"
 #include "Animation/AnimInstance.h"
 #include "Animation/AnimNodeFunctionRef.h"
 #include "Components/SkeletalMeshComponent.h"
 #include "Animation/ExposedValueHandler.h"
+#include "CurveExpressionModule.h"
 #include "Engine/SkeletalMesh.h"
 #include "Animation/AnimCurveTypes.h"
 
 #include UE_INLINE_GENERATED_CPP_BY_NAME(AnimNode_RemapCurvesFromMesh)
 
+void FAnimNode_RemapCurvesFromMesh::Initialize_AnyThread(
+	const FAnimationInitializeContext& Context
+	)
+{
+	DECLARE_SCOPE_HIERARCHICAL_COUNTER_ANIMNODE(Initialize_AnyThread)
+	Super::Initialize_AnyThread(Context);
+	SourcePose.Initialize(Context);
+}
 
-void FAnimNode_RemapCurvesFromMesh::PreUpdate(
-	const UAnimInstance* InAnimInstance
+
+void FAnimNode_RemapCurvesFromMesh::CacheBones_AnyThread(
+	const FAnimationCacheBonesContext& Context
 	)
 {
-	DECLARE_SCOPE_HIERARCHICAL_COUNTER_ANIMNODE(PreUpdate);
+	DECLARE_SCOPE_HIERARCHICAL_COUNTER_ANIMNODE(CacheBones_AnyThread)
+	Super::CacheBones_AnyThread(Context);
+	SourcePose.CacheBones(Context);
+}
+	
 
-	// Make sure we're using the correct source and target skeleton components, since they
-	// may have changed from underneath us.
-	RefreshMeshComponent(InAnimInstance->GetSkelMeshComponent());
+void FAnimNode_RemapCurvesFromMesh::VerifyExpressions(
+	TFunction<void(FString)> InReportingFunc
+	) const
+{
+	using namespace CurveExpression::Evaluator;
 
-	FAnimNode_RemapCurvesBase::PreUpdate(InAnimInstance);
+	auto ReportAndLog = [InReportingFunc](FString InMessage)
+	{
+		if (InReportingFunc)
+		{
+			InReportingFunc(InMessage);
+		}
+		UE_LOG(LogCurveExpression, Warning, TEXT("%s"), *InMessage);
+	};
 	
-	const USkeletalMeshComponent* CurrentMeshComponent = CurrentlyUsedSourceMeshComponent.IsValid() ? CurrentlyUsedSourceMeshComponent.Get() : nullptr;
+	if (!ExpressionEngine.IsSet())
+	{
+		return;
+	}
+	
+	if (CurveExpressions.IsEmpty())
+	{
+		ReportAndLog(TEXT("No curve expressions set."));
+		return;
+	}
 
-	if (CurrentMeshComponent && CurrentMeshComponent->GetSkeletalMeshAsset() && CurrentMeshComponent->IsRegistered())
+	bool bFoundError = false;
+	const FEngine VerificationEngine(
+		ExpressionEngine->GetConstantValues(),
+		EParseFlags::ValidateConstants);
+
+	for (const TPair<FName, FString>& ExpressionPair: CurveExpressions)
 	{
-		// If our source is running under leader-pose, then get bone data from there
-		if(const USkeletalMeshComponent* LeaderPoseComponent = Cast<USkeletalMeshComponent>(CurrentMeshComponent->LeaderPoseComponent.Get()))
+		TOptional<FParseError> Error = VerificationEngine.Verify(ExpressionPair.Value);
+		if (Error.IsSet())
 		{
-			CurrentMeshComponent = LeaderPoseComponent;
+			ReportAndLog(FString::Printf(TEXT("Expression error in '%s': %s"), *ExpressionPair.Value, *Error->Message));
+			bFoundError = true;
 		}
+	}
 
-		// re-check mesh component validity as it may have changed to leader
-		if(CurrentMeshComponent->GetSkeletalMeshAsset() && CurrentMeshComponent->IsRegistered())
+	if (!bFoundError)
+	{
+		UE_LOG(LogCurveExpression, Display, TEXT("Curve expressions verified ok."))
+	}
+}
+
+
+
+bool FAnimNode_RemapCurvesFromMesh::CanVerifyExpressions() const
+{
+	return ExpressionEngine.IsSet();
+}
+
+
+void FAnimNode_RemapCurvesFromMesh::Update_AnyThread(
+	const FAnimationUpdateContext& Context
+	)
+{
+	DECLARE_SCOPE_HIERARCHICAL_COUNTER_ANIMNODE(Update_AnyThread)
+
+	// Run update on input pose nodes
+	SourcePose.Update(Context);
+
+	// Evaluate any BP logic plugged into this node
+	GetEvaluateGraphExposedInputs().Execute(Context);
+	
+	if (bExpressionsImmutable && ExpressionEngine.IsSet() && !CurveExpressions.IsEmpty() && CachedExpressions.IsEmpty())
+	{
+		CachedExpressions.Reset();
+		for (const TPair<FName, FString>& ExpressionPair: CurveExpressions)
 		{
-			if (const UAnimInstance* SourceAnimInstance = CurrentMeshComponent->GetAnimInstance())
-			{
-				// We have a valid instance, let's grab any curve values for the constants used by our expressions.
-				SourceCurveValues.Reset();
+			using namespace CurveExpression::Evaluator;
 
-				const TMap<FName, float>& CurveValues = SourceAnimInstance->GetAnimationCurveList(EAnimCurveType::AttributeCurve);
-				for (FName ConstantName: GetCompiledExpressionConstants())
-				{
-					if (const float* Value = CurveValues.Find(ConstantName))
-					{
-						SourceCurveValues.Add(ConstantName, *Value);
-					}
-				}
+			TVariant<FExpressionObject, FParseError> Result = ExpressionEngine->Parse(ExpressionPair.Value);
+			if (const FExpressionObject* Expression = Result.TryGet<FExpressionObject>())
+			{
+				CachedExpressions.Add(ExpressionPair.Key, *Expression);
 			}
 		}
-		else
-		{
-			// If there's no skeletal mesh, then empty the curve values.
-			SourceCurveValues.Reset();
-		}
-	}
+	}			
 }
 
 
@@ -67,28 +120,42 @@ void FAnimNode_RemapCurvesFromMesh::Evaluate_AnyThread(
 	)
 {
 	DECLARE_SCOPE_HIERARCHICAL_COUNTER_ANIMNODE(Evaluate_AnyThread)
-	
-	using namespace CurveExpression::Evaluator;
-	
+
 	FPoseContext SourceData(Output);
 	SourcePose.Evaluate(SourceData);
 	
 	Output = SourceData;
-
-	for (const TTuple<FName, FExpressionObject>& Assignment: GetCompiledAssignments())
+	
+	// If we have an expression engine, evaluate the expressions that have a matching target curve.
+	// If the expressions are not immutable between compiles, then we need to reparse them each time. 
+	if (ExpressionEngine.IsSet())
 	{
-		if (const SmartName::UID_Type* UID = CurveNameToUIDMap.Find(Assignment.Key))
+		if (bExpressionsImmutable)
+		{
+			FBlendedCurve Curve;
+			Curve.Reserve(CachedExpressions.Num());
+			for (const TTuple<FName, CurveExpression::Evaluator::FExpressionObject>& ExpressionItem : CachedExpressions)
+			{
+				Curve.Add(ExpressionItem.Key, ExpressionEngine->Execute(ExpressionItem.Value)); 
+			}
+
+			Output.Curve.Combine(Curve);
+		}
+		else
 		{
-			Output.Curve.Set(*UID, FEngine().Execute(Assignment.Value,
-				[&CurveValues = SourceCurveValues](const FName InName) -> TOptional<float>
+			FBlendedCurve Curve;
+			Curve.Reserve(CurveExpressions.Num());
+
+			for (const TTuple<FName, FString>& ExpressionItem : CurveExpressions)
+			{
+				TOptional<float> Result = ExpressionEngine->Evaluate(ExpressionItem.Value);
+				if (Result.IsSet())
 				{
-					if (const float* Value = CurveValues.Find(InName))
-					{
-						return *Value;
-					}
-					return {};
+					Curve.Add(ExpressionItem.Key, *Result);
 				}
-			));
+			}
+
+			Output.Curve.Combine(Curve);
 		}
 	}
 }
@@ -106,17 +173,42 @@ void FAnimNode_RemapCurvesFromMesh::GatherDebugData(
 }
 
 
-bool FAnimNode_RemapCurvesFromMesh::Serialize(FArchive& Ar)
+void FAnimNode_RemapCurvesFromMesh::PreUpdate(
+	const UAnimInstance* InAnimInstance
+	)
 {
-	SerializeNode(Ar, this, StaticStruct());
+	QUICK_SCOPE_CYCLE_COUNTER(FAnimNode_RemapCurvesFromMesh_PreUpdate);
 
-	if (IsLoading() && Ar.CustomVer(FCurveExpressionCustomVersion::GUID) < FCurveExpressionCustomVersion::SerializedExpressions)
+	// Make sure we're using the correct source and target skeleton components, since they
+	// may have changed from underneath us.
+	RefreshMeshComponent(InAnimInstance->GetSkelMeshComponent());
+
+	const USkeletalMeshComponent* CurrentMeshComponent = CurrentlyUsedSourceMeshComponent.IsValid() ? CurrentlyUsedSourceMeshComponent.Get() : nullptr;
+
+	if (CurrentMeshComponent && CurrentMeshComponent->GetSkeletalMeshAsset() && CurrentMeshComponent->IsRegistered())
 	{
-		// Default to expression map if loading from an older version.
-		ExpressionSource = ERemapCurvesExpressionSource::ExpressionMap;
+		// If our source is running under leader-pose, then get bone data from there
+		if(USkeletalMeshComponent* LeaderPoseComponent = Cast<USkeletalMeshComponent>(CurrentMeshComponent->LeaderPoseComponent.Get()))
+		{
+			CurrentMeshComponent = LeaderPoseComponent;
+		}
+
+		// re-check mesh component validity as it may have changed to leader
+		if(CurrentMeshComponent->GetSkeletalMeshAsset() && CurrentMeshComponent->IsRegistered())
+		{
+			if (ExpressionEngine.IsSet())
+			{
+				if (const UAnimInstance* SourceAnimInstance = CurrentMeshComponent->GetAnimInstance())
+				{
+					ExpressionEngine->UpdateConstantValues(SourceAnimInstance->GetAnimationCurveList(EAnimCurveType::AttributeCurve));
+				}
+			}
+		}
+		else
+		{
+			ExpressionEngine.Reset();
+		}
 	}
-	
-	return true;
 }
 
 
@@ -129,6 +221,8 @@ void FAnimNode_RemapCurvesFromMesh::ReinitializeMeshComponent(
 	CurrentlyUsedSourceMesh.Reset();
 	CurrentlyUsedTargetMesh.Reset();
 	
+	ExpressionEngine.Reset();
+	CachedExpressions.Reset();
 
 	if (InTargetMeshComponent && IsValid(InNewSkeletalMeshComponent) && InNewSkeletalMeshComponent->GetSkeletalMeshAsset())
 	{
@@ -142,6 +236,19 @@ void FAnimNode_RemapCurvesFromMesh::ReinitializeMeshComponent(
 			CurrentlyUsedSourceMesh = SourceSkelMesh;
 			CurrentlyUsedTargetMesh = TargetSkelMesh;
 
+			// Grab the source curves and use their names to seed the expression evaluation engine.
+			TMap<FName, float> SourceCurves;
+			const USkeleton* SourceSkeleton = SourceSkelMesh->GetSkeleton();
+			if (ensureMsgf(SourceSkeleton, TEXT("Invalid null source skeleton : %s"), *GetNameSafe(TargetSkelMesh)))
+			{
+				// TODO: is this good enough for future use? ExpressionEngine appears to have been built on the assumption that curves are always pre-defined.
+				SourceSkeleton->ForEachCurveMetaData([&SourceCurves](FName InCurveName, const FCurveMetaData& InCurveMetaData)
+				{
+					SourceCurves.Add(InCurveName, 0.0f);
+				});
+
+				ExpressionEngine.Emplace(MoveTemp(SourceCurves));
+			}
 		}
 	}
 }
diff --git a/Engine/Plugins/Experimental/Animation/CurveExpression/Source/Runtime/Private/CurveExpressionCustomVersion.cpp b/Engine/Plugins/Experimental/Animation/CurveExpression/Source/Runtime/Private/CurveExpressionCustomVersion.cpp
deleted file mode 100644
diff --git a/Engine/Plugins/Experimental/Animation/CurveExpression/Source/Runtime/Private/CurveExpressionCustomVersion.h b/Engine/Plugins/Experimental/Animation/CurveExpression/Source/Runtime/Private/CurveExpressionCustomVersion.h
deleted file mode 100644
diff --git a/Engine/Plugins/Experimental/Animation/CurveExpression/Source/Runtime/Private/CurveExpressionsDataAsset.cpp b/Engine/Plugins/Experimental/Animation/CurveExpression/Source/Runtime/Private/CurveExpressionsDataAsset.cpp
deleted file mode 100644
diff --git a/Engine/Plugins/Experimental/Animation/CurveExpression/Source/Runtime/Private/ExpressionEvaluator.cpp b/Engine/Plugins/Experimental/Animation/CurveExpression/Source/Runtime/Private/ExpressionEvaluator.cpp
--- a/Engine/Plugins/Experimental/Animation/CurveExpression/Source/Runtime/Private/ExpressionEvaluator.cpp
+++ b/Engine/Plugins/Experimental/Animation/CurveExpression/Source/Runtime/Private/ExpressionEvaluator.cpp
@@ -3,8 +3,6 @@
 #include "ExpressionEvaluator.h"
 #include <limits> // IWYU pragma: keep
 
-#include "CurveExpressionModule.h"
-
 namespace CurveExpression::Evaluator
 {
 
@@ -106,124 +104,63 @@ FString FExpressionObject::ToString() const
 }
 
 
-TSet<FName> FExpressionObject::GetUsedConstants() const
+FEngine::FEngine(
+	const TMap<FName, float>& InConstants, 
+	EParseFlags InParseFlags
+	) :
+	Constants(InConstants),
+	ParseFlags(InParseFlags)
 {
-	TSet<FName> UsedConstants;
-	for (int32 Index = 0; Index < Expression.Num(); Index++)
-	{
-		if (const FName* ConstantName = Expression[Index].TryGet<FName>())
-		{
-			UsedConstants.Add(*ConstantName);
-		}
-	}
-	return UsedConstants;	
 }
 
 
-void FExpressionObject::Load(FArchive& Ar)
+FEngine::FEngine(
+	TMap<FName, float>&& InConstants,
+	EParseFlags InParseFlags
+	) :
+	Constants(MoveTemp(InConstants)),
+	ParseFlags(InParseFlags)
 {
-	int32 OperandCount = 0;
-	Ar << OperandCount;
-
-	Expression.Reset(OperandCount);
-
-	for (int32 OperandIndex = 0; OperandIndex < OperandCount; OperandIndex++)
-	{
-		int32 OperandType = INDEX_NONE;
-		Ar << OperandType;
-		switch(OperandType)
-		{
-		case OpElement::IndexOfType<EOperator>():
-			{
-				int32 OperatorType = 0;
-				Ar << OperatorType;
-				Expression.Add(OpElement(TInPlaceType<EOperator>(), static_cast<EOperator>(OperatorType)));
-			}
-			break;
-		case OpElement::IndexOfType<FName>():
-			{
-				FName ConstantName;
-				Ar << ConstantName;
-				Expression.Add(OpElement(TInPlaceType<FName>(), ConstantName));
-			}
-			break;
-		case OpElement::IndexOfType<FFunctionRef>():
-			{
-				int32 FunctionIndex;
-				Ar << FunctionIndex;
-				Expression.Add(OpElement(TInPlaceType<FFunctionRef>(), FFunctionRef(FunctionIndex)));
-			}
-			break;
-		case OpElement::IndexOfType<float>():
-			{
-				float Value;
-				Ar << Value;
-				Expression.Add(OpElement(TInPlaceType<float>(), Value));
-			}
-			break;
-		default:
-			UE_LOG(LogCurveExpression, Error, TEXT("Invalid element type found in serialized expression"));
-			Ar.SetError();
-			return;
-		}
-	}
 }
 
 
-void FExpressionObject::Save(FArchive& Ar) const
+void FEngine::UpdateConstantValues(
+	const TMap<FName, float>& InConstants,
+	EUpdateConstantsMethod InUpdateMethod
+	)
 {
-	int32 OperandCount = Expression.Num(); 
-	Ar << OperandCount;
-
-	for (const OpElement& Operand: Expression)
+	for (TTuple<FName, float>& OldConstantItem: Constants)
 	{
-		int32 OperandType = Operand.GetIndex();
-		Ar << OperandType;
-		
-		switch(OperandType)
+		if (const float* NewConstant = InConstants.Find(OldConstantItem.Key))
 		{
-		case OpElement::IndexOfType<EOperator>():
-			{
-				int32 OperatorType = static_cast<int32>(Operand.Get<EOperator>());
-				Ar << OperatorType;
-			}
-			break;
-		case OpElement::IndexOfType<FName>():
-			{
-				FName ConstantName = Operand.Get<FName>();
-				Ar << ConstantName;
-			}
-			break;
-		case OpElement::IndexOfType<FFunctionRef>():
-			{
-				int32 FunctionIndex = Operand.Get<FFunctionRef>().Index;
-				Ar << FunctionIndex;
-			}
-			break;
-		case OpElement::IndexOfType<float>():
-			{
-				float Value = Operand.Get<float>();
-				Ar << Value;
-			}
-			break;
-		default:
-			checkNoEntry();
-			return;
+			OldConstantItem.Value = *NewConstant;
+		}
+		else if (InUpdateMethod == EUpdateConstantsMethod::ZeroMissing)
+		{
+			OldConstantItem.Value = 0.0f;
 		}
 	}
 }
 
-
 // We can't use FChar::IsDigit, since it obeys locale, which would allow all sorts of non-ASCII
 // numeric digits from Unicode.
-static bool IsASCIIDigit(const TCHAR InCh)
+static bool IsASCIIDigit(TCHAR InCh)
 {
 	return InCh >= '0' && InCh <= '9';
 }
 
+TVariant<FEngine::FToken, FParseError> FEngine::ParseTokenError(
+	FString&& InError,
+	const TCHAR* InCurrentPos,
+	FStringView InExpression)
+{
+	return TVariant<FToken, FParseError>(TInPlaceType<FParseError>(),
+		static_cast<int32>(InCurrentPos - InExpression.GetData()), MoveTemp(InError));
+}
+
 
 FString FEngine::TokenToString(
-	const FToken& InToken
+	const FEngine::FToken& InToken
 	)
 {
 	const TCHAR* TypeStr = TEXT("<unknown>");
@@ -245,10 +182,6 @@ FString FEngine::TokenToString(
 		case EOperatorToken::ParenOpen:		ValueStr = TEXT("("); break;
 		case EOperatorToken::ParenClose:	ValueStr = TEXT(")"); break;
 		case EOperatorToken::Comma:			ValueStr = TEXT(","); break;
-		
-		case EOperatorToken::Max:
-			checkNoEntry();
-			break;
 		}
 	}
 	else if (const FName* Identifier = InToken.TryGet<FName>())
@@ -266,10 +199,10 @@ FString FEngine::TokenToString(
 }
 
 
-auto FEngine::ParseFloat(
+TVariant<FEngine::FToken, FParseError> FEngine::ParseFloat(
 	FStringView& InOutParseRange,
 	FStringView InExpression
-	) -> FTokenParseResult
+	)
 {
 	// We only parse at most 11 significant digits, since 9 is the maximum decimal precision 
 	// a float will hold. We store 2 extra digits to make sure we round properly on the 
@@ -325,9 +258,7 @@ auto FEngine::ParseFloat(
 	if (MantissaLen == 0)
 	{
 		// We only parsed the period. That's not a valid float.
-		return FTokenParseResult(TEXT("Invalid floating point value"),
-			InOutParseRange.Left(1), 
-		InExpression.GetData());
+		return ParseTokenError(TEXT("Invalid floating point value"), InOutParseRange.GetData(), InExpression);
 	}
 
 	// Adjustment for the exponent based on number of significant digits in the fraction part.
@@ -441,19 +372,17 @@ auto FEngine::ParseFloat(
 		Fraction *= DoubleExponent;
 	}
 
-	const FStringView TokenParseRange(InOutParseRange.Left(LastPos));
-	
-	// Update the parse range for the next token. 
 	InOutParseRange = InOutParseRange.Mid(LastPos);
 	
-	return FTokenParseResult(FToken(TInPlaceType<float>(), Fraction), TokenParseRange, InExpression.GetData());
+	return TVariant<FToken, FParseError>(
+		TInPlaceType<FToken>(), FToken(TInPlaceType<float>(), Fraction));
 }
 
 
-auto FEngine::ParseIdentifier(
+TVariant<FEngine::FToken, FParseError> FEngine::ParseIdentifier(
 	FStringView& InOutParseRange,
 	FStringView InExpression
-	) -> FTokenParseResult 
+	) const
 {
 	// FIXME: We need quote escaping.
 	bool bStopAtSingleQuote = false;
@@ -464,13 +393,13 @@ auto FEngine::ParseIdentifier(
 		ParseIndex++;
 		if (ParseIndex >= InOutParseRange.Len())
 		{
-			return FTokenParseResult(TEXT("Missing end quote for constant"), InOutParseRange, InExpression.GetData());
+			return ParseTokenError(TEXT("Missing end quote for constant"), InOutParseRange.GetData(), InExpression);
 		}
 	}
 	// Unquoted constants have to start with an alphabetic character.
 	else if (!FChar::IsAlpha(InOutParseRange[0]) && InOutParseRange[0] != '_')
 	{
-		return FTokenParseResult(TEXT("Unexpected character"), InOutParseRange.Left(1), InExpression.GetData());
+		return ParseTokenError(TEXT("Unexpected character"), InOutParseRange.GetData(), InExpression);
 	}
 
 	const int32 ConstantStart = ParseIndex;
@@ -487,7 +416,7 @@ auto FEngine::ParseIdentifier(
 
 		if (ParseIndex == InOutParseRange.Len())
 		{
-			return FTokenParseResult(TEXT("Missing end quote for constant"), InOutParseRange, InExpression.GetData());
+			return ParseTokenError(TEXT("Missing end quote for constant"), InOutParseRange.GetData(), InExpression);
 		}
 
 		ConstantEnd = ParseIndex;
@@ -504,61 +433,86 @@ auto FEngine::ParseIdentifier(
 		ConstantEnd = ParseIndex;
 	}
 
+	// Look up this identifier. We don't want to create a new FName, to avoid polluting the
+	// FNameEntry cache.
 	const FStringView Identifier(InOutParseRange.SubStr(ConstantStart, ConstantEnd - ConstantStart));
-	const FName Name(Identifier.Len(), Identifier.GetData());
+	const FName Name(Identifier.Len(), Identifier.GetData(), FNAME_Find);
 	
-	const FStringView TokenParseRange(InOutParseRange.Left(ParseIndex));
-	InOutParseRange = InOutParseRange.Mid(ParseIndex);
+	InOutParseRange = InOutParseRange.Mid(ParseIndex); 
 
-	return FTokenParseResult(FToken(TInPlaceType<FName>(), Name), TokenParseRange, InExpression.GetData());
+	if (Constants.Contains(Name) || FunctionNameIndex.Contains(Name))
+	{
+		return TVariant<FToken, FParseError>(
+			TInPlaceType<FToken>(), FToken(TInPlaceType<FName>(), Name));
+	}
+	else if (EnumHasAllFlags(ParseFlags, EParseFlags::ValidateConstants))
+	{
+		return ParseTokenError(FString::Printf(TEXT("Unknown identifier '%s'"), *FString(Identifier)),
+			InOutParseRange.GetData(), InExpression); 
+	}
+	else
+	{
+		// If the constant value is not found, it gets a default value of None.
+		return TVariant<FToken, FParseError>(
+			TInPlaceType<FToken>(), FToken(TInPlaceType<FName>(), NAME_None));
+	}
 }
 
 
-auto FEngine::ParseToken(
+TVariant<FEngine::FToken, FParseError> FEngine::ParseToken(
 	FStringView& InOutParseRange,
 	FStringView InExpression
-	) -> FTokenParseResult
+	) const
 {
-	auto TokenResult = [&InOutParseRange, Expression=InExpression.GetData()](EOperatorToken InTokenKind, int32 InTokenLength)
+	auto TokenResult = [](EOperatorToken InTokenKind)
 	{
-		const FStringView TokenRange(InOutParseRange.Left(InTokenLength));
-		InOutParseRange = InOutParseRange.Mid(InTokenLength); 
-		return FTokenParseResult(FToken(TInPlaceType<EOperatorToken>(), InTokenKind), TokenRange, Expression);
+		return TVariant<FToken, FParseError>(
+			TInPlaceType<FToken>(), FToken(TInPlaceType<EOperatorToken>(), InTokenKind));
 	};
 
 	switch(InOutParseRange[0])
 	{
 	case '+':
-		return TokenResult(EOperatorToken::Add, 1);
+		InOutParseRange = InOutParseRange.Mid(1);
+		return TokenResult(EOperatorToken::Add);
 		
 	case '-':
-		return TokenResult(EOperatorToken::Subtract, 1);
+		InOutParseRange = InOutParseRange.Mid(1);
+		return TokenResult(EOperatorToken::Subtract);
 		
 	case '*':
 		if (InOutParseRange.Len() >= 2 && InOutParseRange[1] == '*')
 		{
-			return TokenResult(EOperatorToken::Power, 2);
+			InOutParseRange = InOutParseRange.Mid(2);
+			return TokenResult(EOperatorToken::Power);
 		}
-		return TokenResult(EOperatorToken::Multiply, 1);
+		InOutParseRange = InOutParseRange.Mid(1);
+		return TokenResult(EOperatorToken::Multiply);
 		
 	case '/':
 		if (InOutParseRange.Len() >= 2 && InOutParseRange[1] == '/')
 		{
-			return TokenResult(EOperatorToken::FloorDivide, 2);
+			InOutParseRange = InOutParseRange.Mid(2);
+			return TokenResult(EOperatorToken::FloorDivide);
 		}
-		return TokenResult(EOperatorToken::Divide, 1);
+		InOutParseRange = InOutParseRange.Mid(1);
+		return TokenResult(EOperatorToken::Divide);
 		
 	case '%':
-		return TokenResult(EOperatorToken::Modulo, 1);
+		InOutParseRange = InOutParseRange.Mid(1);
+		return TokenResult(EOperatorToken::Modulo);
 
 	case '(':
-		return TokenResult(EOperatorToken::ParenOpen, 1);
+		InOutParseRange = InOutParseRange.Mid(1);
+		return TokenResult(EOperatorToken::ParenOpen);
 		
 	case ')':
-		return TokenResult(EOperatorToken::ParenClose, 1);
+		InOutParseRange = InOutParseRange.Mid(1);
+		return TokenResult(EOperatorToken::ParenClose);
 
 	case ',':
-		return TokenResult(EOperatorToken::Comma, 1);
+		InOutParseRange = InOutParseRange.Mid(1);
+		return TokenResult(EOperatorToken::Comma);
 		
 	default:
 		if (IsASCIIDigit(InOutParseRange[0]) || InOutParseRange[0] == '.')
@@ -571,8 +525,7 @@ auto FEngine::ParseToken(
 
 
 TVariant<FExpressionObject, FParseError> FEngine::Parse(
-	FStringView InExpression,
-	TFunction<TOptional<float>(FName InConstantName)> InConstantEvaluator
+	FStringView InExpression
 	) const
 {
 	enum class EAssociativity
@@ -620,44 +573,44 @@ TVariant<FExpressionObject, FParseError> FEngine::Parse(
 	// Run Dijkstra's classic Shunting Yard algorithm to convert infix expressions to RPN.
 	// TODO: Better error pinpointing by storing the token parse result with the token.
 	TArray<FExpressionObject::OpElement, TInlineAllocator<64>> Expression;
-	TArray<FParseLocation, TInlineAllocator<64>> Locations;
-	TArray<TTuple<EOperatorToken, FParseLocation>, TInlineAllocator<32>> OperatorStack;
-	
+	TArray<EOperatorToken, TInlineAllocator<32>> OperatorStack;
 	struct FFunctionCallInfo
 	{
 		int32 FunctionIndex = INDEX_NONE;
 		int32 CountedCommas = 0;
 		int32 OpeningOpStackSize = 0;
 		int32 ExpressionSize = 0;
-		FParseLocation Location{};
+		const TCHAR* TokenStart = nullptr;
 	};
 	TArray<FFunctionCallInfo, TInlineAllocator<32>> FunctionStack;
 
-	auto ParseError = [](FString&& InError, FParseLocation InLocation)
+	auto ParseError = [InExpression](FString&& InError, const TCHAR* InCurrentPos)
 	{
-		return TVariant<FExpressionObject, FParseError>(TInPlaceType<FParseError>(), MoveTemp(InError), InLocation);
+		return TVariant<FExpressionObject, FParseError>(TInPlaceType<FParseError>(),
+			static_cast<int32>(InCurrentPos - InExpression.GetData()), MoveTemp(InError));
 	};
-	auto PushOperatorTokenToExpression = [GetOperatorTokenInfo, &Expression, &Locations](TTuple<EOperatorToken, FParseLocation> InTokenAndLocation)
+	auto PushOperatorTokenToExpression = [GetOperatorTokenInfo, &Expression](EOperatorToken InOpToken)
 	{
 		Expression.Push(FExpressionObject::OpElement(TInPlaceType<FExpressionObject::EOperator>(), 
-			*GetOperatorTokenInfo(InTokenAndLocation.Get<0>()).Operator));
-		Locations.Push(InTokenAndLocation.Get<1>());
+			*GetOperatorTokenInfo(InOpToken).Operator));
 	};
 	// Used to check if we encounter a '-' to see if we need to turn it into a negate operator
 	// and also to see if we have two operators in a row (error).
 	TOptional<FToken> LastToken;
+	FStringView LastTokenParseRange;
 
 	do
 	{
-		FTokenParseResult ParseResult = ParseToken(ParseRange, InExpression);
+		const TCHAR* TokenStart = ParseRange.GetData();
+		TVariant<FToken, FParseError> ParseResult = ParseToken(ParseRange, InExpression);
 		
-		if (ParseResult.IsError())
+		if (FParseError* Error = ParseResult.TryGet<FParseError>())
 		{
-			return TVariant<FExpressionObject, FParseError>(TInPlaceType<FParseError>(), ParseResult.GetParseError());
+			return TVariant<FExpressionObject, FParseError>(TInPlaceType<FParseError>(), *Error);
 		}
 
-		const FToken& Token = ParseResult.GetToken();
-		
+		const FStringView TokenParseRange(TokenStart, ParseRange.GetData() - TokenStart);
+		const FToken& Token = ParseResult.Get<FToken>();
 
 		// Values / identifiers go onto the expression stack immediately.
 		if (const FName* Identifier = Token.TryGet<FName>())
@@ -665,20 +618,18 @@ TVariant<FExpressionObject, FParseError> FEngine::Parse(
 			if (LastToken.IsSet() &&
 				(LastToken->IsType<FName>() || LastToken->IsType<float>()))
 			{
-				return ParseError(TEXT("Expected an operator"), ParseResult.Location);
+				return ParseError(TEXT("Expected an operator"), TokenParseRange.GetData());
 			}
 			Expression.Push(FExpressionObject::OpElement(TInPlaceType<FName>(), *Identifier));
-			Locations.Push(ParseResult.Location);
 		}
 		else if (const float* Value = Token.TryGet<float>())
 		{
 			if (LastToken.IsSet() &&
 				(LastToken->IsType<FName>() || LastToken->IsType<float>()))
 			{
-				return ParseError(TEXT("Expected an operator"), ParseResult.Location);
+				return ParseError(TEXT("Expected an operator"), TokenParseRange.GetData());
 			}
 			Expression.Push(FExpressionObject::OpElement(TInPlaceType<float>(), *Value));
-			Locations.Push(ParseResult.Location);
 		}
 		else
 		{
@@ -693,16 +644,16 @@ TVariant<FExpressionObject, FParseError> FEngine::Parse(
 					const FOperatorTokenInfo& LastTokenInfo = GetOperatorTokenInfo(*LastOp);
 					if (TokenInfo.Operator.IsSet() && LastTokenInfo.Operator.IsSet())
 					{
-						return ParseError(TEXT("Expected an expression"), ParseResult.Location);
+						return ParseError(TEXT("Expected an expression"), TokenParseRange.GetData());
 					}
 					if (*LastOp == EOperatorToken::ParenClose && Op == EOperatorToken::ParenOpen)
 					{
-						return ParseError(TEXT("Expected an operator"), ParseResult.Location);
+						return ParseError(TEXT("Expected an operator"), TokenParseRange.GetData());
 					}
 					if (*LastOp == EOperatorToken::ParenOpen && Op == EOperatorToken::ParenClose &&
 					    (FunctionStack.IsEmpty() || FunctionStack.Top().OpeningOpStackSize != OperatorStack.Num()))
 					{
-						return ParseError(TEXT("Empty Parentheses"), ParseResult.Location);
+						return ParseError(TEXT("Empty Parentheses"), TokenParseRange.GetData());
 					}
 				}
 			}
@@ -712,35 +663,34 @@ TVariant<FExpressionObject, FParseError> FEngine::Parse(
 				(!LastToken.IsSet() || (LastToken->IsType<EOperatorToken>() && LastToken->Get<EOperatorToken>() != EOperatorToken::ParenClose)))
 			{
 				Op = EOperatorToken::Negate;
-				OperatorStack.Push({Op, ParseResult.Location});
+				OperatorStack.Push(Op);
 			}
 			else if (Op == EOperatorToken::ParenOpen)
 			{
-				OperatorStack.Push({Op, ParseResult.Location});
+				OperatorStack.Push(Op);
 
-				// If the top element on the expression is a constant, check if it is a function name and if so, 
-				// pop the constant out of the expression. 
+				// If the top element on the expression is a constant, check if it is a function and if so, pop the constant
+				// out of the expression. 
 				if (!Expression.IsEmpty() && Expression.Top().IsType<FName>())
 				{
 					FName ConstantName = Expression.Top().Get<FName>();
 					const int32 *FunctionIndexPtr = FunctionNameIndex.Find(ConstantName);
 					if (FunctionIndexPtr == nullptr)
 					{
-						return ParseError(FString::Printf(TEXT("Unknown function '%s'"), *ConstantName.ToString()), Locations.Top());
+						return ParseError(FString::Printf(TEXT("Unknown function '%s'"), *ConstantName.ToString()), LastTokenParseRange.GetData());
 					}
 
 					Expression.Pop();
-					const FParseLocation LastLocation = Locations.Pop();
-					FunctionStack.Push({*FunctionIndexPtr, 0, OperatorStack.Num(), Expression.Num(), LastLocation});
+					FunctionStack.Push({*FunctionIndexPtr, 0, OperatorStack.Num(), Expression.Num(), LastTokenParseRange.GetData()});
 				}
 			}
 			else if (Op == EOperatorToken::ParenClose)
 			{
-				while(OperatorStack.IsEmpty() || OperatorStack.Last().Key != EOperatorToken::ParenOpen)
+				while(OperatorStack.IsEmpty() || OperatorStack.Last() != EOperatorToken::ParenOpen)
 				{
 					if (OperatorStack.IsEmpty())
 					{
-						return ParseError(TEXT("Mismatched parentheses"), ParseResult.Location); 
+						return ParseError(TEXT("Mismatched parentheses"), TokenParseRange.GetData()); 
 					}
 					PushOperatorTokenToExpression(OperatorStack.Pop(false));
 				}
@@ -754,12 +704,10 @@ TVariant<FExpressionObject, FParseError> FEngine::Parse(
 
 					if (ArgumentCount != FunctionInfo.ArgumentCount)
 					{
-						return ParseError(FString::Printf(TEXT("Invalid argument count. Expected %d, got %d"), FunctionInfo.ArgumentCount, ArgumentCount),
-							CallInfo.Location.Join(ParseResult.Location)); 
+						return ParseError(FString::Printf(TEXT("Invalid argument count. Expected %d, got %d"), FunctionInfo.ArgumentCount, ArgumentCount), CallInfo.TokenStart); 
 					}
 
 					Expression.Push(FExpressionObject::OpElement(TInPlaceType<FExpressionObject::FFunctionRef>(), CallInfo.FunctionIndex));
-					Locations.Push(ParseResult.Location);
 				}
 				
 				// Remove the now unneeded open parentheses.
@@ -769,58 +717,44 @@ TVariant<FExpressionObject, FParseError> FEngine::Parse(
 			{
 				if (FunctionStack.IsEmpty() || FunctionStack.Top().OpeningOpStackSize != OperatorStack.Num())
 				{
-					return ParseError(TEXT("Unexpected comma"), ParseResult.Location); 
+					return ParseError(TEXT("Unexpected comma"), TokenParseRange.GetData()); 
 				}
 				FunctionStack.Top().CountedCommas++;
 			}
 			else if (GetOperatorTokenInfo(Op).Associativity == EAssociativity::Left)
 			{
 				// Peel off all operators that have the same precedence or higher.
-				while (!OperatorStack.IsEmpty() && GetOperatorTokenInfo(OperatorStack.Last().Get<0>()).Precedence >= TokenInfo.Precedence)
+				while (!OperatorStack.IsEmpty() && GetOperatorTokenInfo(OperatorStack.Last()).Precedence >= TokenInfo.Precedence)
 				{
 					PushOperatorTokenToExpression(OperatorStack.Pop(false));
 				}
-				OperatorStack.Push({Op, ParseResult.Location});
+				OperatorStack.Push(Op);
 			}
 			else if (GetOperatorTokenInfo(Op).Associativity == EAssociativity::Right)
 			{
 				// Peel off all operators that have higher precedence.
-				while (!OperatorStack.IsEmpty() && GetOperatorTokenInfo(OperatorStack.Last().Get<0>()).Precedence > TokenInfo.Precedence)
+				while (!OperatorStack.IsEmpty() && GetOperatorTokenInfo(OperatorStack.Last()).Precedence > TokenInfo.Precedence)
 				{
 					PushOperatorTokenToExpression(OperatorStack.Pop(false));
 				}
-				OperatorStack.Push({Op, ParseResult.Location});
+				OperatorStack.Push(Op);
 			}
 		}
 
 		LastToken = Token;
+		LastTokenParseRange = ParseRange;
 		ParseRange = ParseRange.TrimStart();
 	} while(!ParseRange.IsEmpty());
 
 	while (!OperatorStack.IsEmpty())
 	{
-		if (!GetOperatorTokenInfo(OperatorStack.Last().Get<0>()).Operator.IsSet())
+		if (!GetOperatorTokenInfo(OperatorStack.Last()).Operator.IsSet())
 		{
-			return ParseError(TEXT("Mismatched parentheses"), OperatorStack.Last().Get<1>()); 
+			return ParseError(TEXT("Mismatched parentheses"), InExpression.GetData()); 
 		}
 		PushOperatorTokenToExpression(OperatorStack.Pop(false));
 	}
 
-	// If a constant evaluator function is given, then check if the constants exist.  
-	if (InConstantEvaluator)
-	{
-		for (int32 Index = 0; Index < Expression.Num(); Index++)
-		{
-			if (const FName* ConstantName = Expression[Index].TryGet<FName>())
-			{
-				if (!InConstantEvaluator(*ConstantName).IsSet())
-				{
-					return ParseError(*FString::Printf(TEXT("Unknown identifier '%s'"), *ConstantName->ToString()), Locations[Index]); 
-				}
-			}
-		}
-	}
-	
 	FExpressionObject S;
 	S.Expression = Expression;
 	return TVariant<FExpressionObject, FParseError>(TInPlaceType<FExpressionObject>(), S);
@@ -828,8 +762,7 @@ TVariant<FExpressionObject, FParseError> FEngine::Parse(
 
 
 float FEngine::Execute(
-	const FExpressionObject& InExpressionObject,
-	TFunctionRef<TOptional<float>(FName InConstantName)> InConstantEvaluator
+	const FExpressionObject& InExpressionObject
 	) const
 {
 	// An empty expression object is valid.
@@ -848,10 +781,10 @@ float FEngine::Execute(
 		{
 			ValueStack.Push(*Value);
 		}
-		else if (const FName* ConstantName = Token.TryGet<FName>())
+		else if (const FName* Constant = Token.TryGet<FName>())
 		{
-			TOptional<float> ConstantValue = InConstantEvaluator(*ConstantName);
-			ValueStack.Push(ConstantValue.IsSet() ? ConstantValue.GetValue() : 0.0f);
+			const float *ConstantValue = Constants.Find(*Constant);
+			ValueStack.Push(ConstantValue ? *ConstantValue : 0.0f);
 		}
 		else if (const FExpressionObject::EOperator* Operator = Token.TryGet<FExpressionObject::EOperator>())
 		{
@@ -953,14 +886,13 @@ float FEngine::Execute(
 
 
 TOptional<float> FEngine::Evaluate(
-	FStringView InExpression,
-	TFunctionRef<TOptional<float>(FName InConstantName)> InConstantEvaluator
+	FStringView InExpression
 	) const
 {
 	TVariant<FExpressionObject, FParseError> Result = Parse(InExpression);
 	if (const FExpressionObject* ExpressionObject = Result.TryGet<FExpressionObject>())
 	{
-		return Execute(*ExpressionObject, InConstantEvaluator);
+		return Execute(*ExpressionObject);
 	}
 
 	return {};
@@ -968,11 +900,10 @@ TOptional<float> FEngine::Evaluate(
 
 
 TOptional<FParseError> FEngine::Verify(
-	FStringView InExpression,
-	TFunction<TOptional<float>(FName InConstantName)> InConstantEvaluator 
+	FStringView InExpression
 	) const
 {
-	TVariant<FExpressionObject, FParseError> Result = Parse(InExpression, InConstantEvaluator);
+	TVariant<FExpressionObject, FParseError> Result = Parse(InExpression);
 	if (const FParseError* Error = Result.TryGet<FParseError>())
 	{
 		return *Error;
diff --git a/Engine/Plugins/Experimental/Animation/CurveExpression/Source/Runtime/Public/AnimNode_RemapCurves.h b/Engine/Plugins/Experimental/Animation/CurveExpression/Source/Runtime/Public/AnimNode_RemapCurves.h
deleted file mode 100644
diff --git a/Engine/Plugins/Experimental/Animation/CurveExpression/Source/Runtime/Public/AnimNode_RemapCurvesFromMesh.h b/Engine/Plugins/Experimental/Animation/CurveExpression/Source/Runtime/Public/AnimNode_RemapCurvesFromMesh.h
--- a/Engine/Plugins/Experimental/Animation/CurveExpression/Source/Runtime/Public/AnimNode_RemapCurvesFromMesh.h
+++ b/Engine/Plugins/Experimental/Animation/CurveExpression/Source/Runtime/Public/AnimNode_RemapCurvesFromMesh.h
@@ -2,51 +2,64 @@
 
 #pragma once
 
-#include "AnimNode_RemapCurvesBase.h"
+#include "ExpressionEvaluator.h"
+#include "Animation/AnimNodeBase.h"
 
 #include "AnimNode_RemapCurvesFromMesh.generated.h"
 
-
 USTRUCT(BlueprintInternalUseOnly)
 struct CURVEEXPRESSION_API FAnimNode_RemapCurvesFromMesh :
-	public FAnimNode_RemapCurvesBase
+	public FAnimNode_Base
 {
 	GENERATED_BODY()
 
+	UPROPERTY(EditAnywhere, EditFixedSize, BlueprintReadWrite, Category = Links)
+	FPoseLink SourcePose;
+	
 	/** This is used by default if it's valid */
-	UPROPERTY(BlueprintReadWrite, transient, Category=Copy, meta=(PinShownByDefault, DisplayAfter="SourcePose"))
+	UPROPERTY(BlueprintReadWrite, transient, Category=Copy, meta=(PinShownByDefault))
 	TWeakObjectPtr<USkeletalMeshComponent> SourceMeshComponent;
 
 	/** If SourceMeshComponent is not valid, and if this is true, it will look for attached parent as a source */
 	UPROPERTY(EditAnywhere, BlueprintReadWrite, Category=Copy, meta=(NeverAsPin))
 	bool bUseAttachedParent = false;
 	
+	UPROPERTY(EditAnywhere, BlueprintReadWrite, editfixedsize, Category = Expressions, meta = (PinHiddenByDefault))
+	TMap<FName, FString> CurveExpressions;
+
+	/** The expressions given are immutable and will not change during runtime. Improves performance. */ 
+	UPROPERTY(EditAnywhere, BlueprintReadWrite, Category=Expressions, meta=(NeverAsPin))
+	bool bExpressionsImmutable = true;
+
+	// Call to verify the expressions. Write out any warnings to the InReportingFunc function.
+	void VerifyExpressions(
+		TFunction<void(FString)> InReportingFunc
+		) const;
+	
+	// Returns true if we're able to verify expressions at all. I.e. if the expression engine has not been initialized.
+	bool CanVerifyExpressions() const;
+
+	
 	// FAnimNode_Base interface
-	virtual void PreUpdate(const UAnimInstance* InAnimInstance) override;
+	virtual void Initialize_AnyThread(const FAnimationInitializeContext& Context) override;
+	virtual void CacheBones_AnyThread(const FAnimationCacheBonesContext& Context) override;
+	virtual void Update_AnyThread(const FAnimationUpdateContext& Context) override;
 	virtual void Evaluate_AnyThread(FPoseContext& Output) override;
 	virtual void GatherDebugData(FNodeDebugData& DebugData) override;
-
-	bool Serialize(FArchive& Ar);
+	virtual bool HasPreUpdate() const override { return true; }
+	virtual void PreUpdate(const UAnimInstance* InAnimInstance) override;
 	
 private:
 	void ReinitializeMeshComponent(USkeletalMeshComponent* InNewSkeletalMeshComponent, USkeletalMeshComponent* InTargetMeshComponent);
 	void RefreshMeshComponent(USkeletalMeshComponent* InTargetMeshComponent);
 
 	// this is source mesh references, so that we could compare and see if it has changed
-	TWeakObjectPtr<const USkeletalMeshComponent> CurrentlyUsedSourceMeshComponent;
-	TWeakObjectPtr<const USkeletalMesh> CurrentlyUsedSourceMesh;
+	TWeakObjectPtr<USkeletalMeshComponent>	CurrentlyUsedSourceMeshComponent;
+	TWeakObjectPtr<USkeletalMesh>			CurrentlyUsedSourceMesh;
 
 	// target mesh 
-	TWeakObjectPtr<const USkeletalMesh> CurrentlyUsedTargetMesh;
+	TWeakObjectPtr<USkeletalMesh>			CurrentlyUsedTargetMesh;
 
-	// Transient data.
-	TMap<FName, float> SourceCurveValues;
-};
-
-template<> struct TStructOpsTypeTraits<FAnimNode_RemapCurvesFromMesh> : public TStructOpsTypeTraitsBase2<FAnimNode_RemapCurvesFromMesh>
-{
-	enum 
-	{ 
-		WithSerializer = true
-	};
+	TOptional<CurveExpression::Evaluator::FEngine> ExpressionEngine;
+	TMap<FName, CurveExpression::Evaluator::FExpressionObject> CachedExpressions;
 };
diff --git a/Engine/Plugins/Experimental/Animation/CurveExpression/Source/Runtime/Public/CurveExpressionsDataAsset.h b/Engine/Plugins/Experimental/Animation/CurveExpression/Source/Runtime/Public/CurveExpressionsDataAsset.h
deleted file mode 100644
diff --git a/Engine/Plugins/Experimental/Animation/CurveExpression/Source/Runtime/Public/ExpressionEvaluator.h b/Engine/Plugins/Experimental/Animation/CurveExpression/Source/Runtime/Public/ExpressionEvaluator.h
--- a/Engine/Plugins/Experimental/Animation/CurveExpression/Source/Runtime/Public/ExpressionEvaluator.h
+++ b/Engine/Plugins/Experimental/Animation/CurveExpression/Source/Runtime/Public/ExpressionEvaluator.h
@@ -4,6 +4,7 @@
 
 #include "Containers/Map.h"
 #include "Misc/TVariant.h"
+#include "Misc/TVariantMeta.h"
 #include "UObject/NameTypes.h"
 
 template <typename OptionalType> struct TOptional;
@@ -11,55 +12,22 @@ template <typename OptionalType> struct TOptional;
 namespace CurveExpression::Evaluator
 {
 
-struct FParseLocation
-{
-	FParseLocation() = default;
-	FParseLocation(int32 InStart, int32 InEndLocation) :
-		Start(InStart), 
-		End(InEndLocation)
-	{}
-	FParseLocation(FStringView InParseRange, const TCHAR* InParseStart) :
-		Start(static_cast<int32>(InParseRange.begin() - InParseStart)),
-		End(static_cast<int32>(InParseRange.end() - InParseStart))
-	{}
-
-	/** Merge two parse locations and return a new one */
-	FParseLocation Join(const FParseLocation& InLocation) const
-	{
-		return FParseLocation(FMath::Min(Start, InLocation.Start), FMath::Max(End, InLocation.End));
-	}
-	
-	const int32 Start = INDEX_NONE;
-	const int32 End = INDEX_NONE;
-};
-
 struct FParseError
 {
 	FParseError() = default;
-	FParseError(FString&& InMessage, FParseLocation InLocation) :
-		Message(MoveTemp(InMessage)),
-		Location(InLocation)
-	{}
+	FParseError(const int32 InColumn, FString&& InMessage) : Column(InColumn), Message(MoveTemp(InMessage)) {}
+
+	/** The column at which the error occurs. Zero-based. */
+	const int32 Column = 0;
 
 	/** The human-readable parsing error */
-	const FString Message{};
-	
-	/** The parse location where the error took place. */
-	const FParseLocation Location{};
+	const FString Message = FString();
 };
 
 struct FExpressionObject
 {
 	FString ToString() const;
-
-	/** Returns all named constants used by this expression */
-	TSet<FName> GetUsedConstants() const;
-
-	void Load(FArchive& Ar);
-	void Save(FArchive& Ar) const;
-	
 private:
-	
 	enum class EOperator : int32
 	{
 		Negate,				// Negation operator.
@@ -106,40 +74,74 @@ private:
  *	rather than:
  *	   foo (1) * 2.0
  */
+enum class EParseFlags
+{
+	Default				= 0,		/** Default parsing behavior */
+	ValidateConstants	= 1 << 0,	/** Validate that the referred-to constant exists */
+};
+ENUM_CLASS_FLAGS(EParseFlags)
+
 
+/** Dictate how constants are updated. */
+enum class EUpdateConstantsMethod
+{
+	IgnoreMissing,					/** Existing constant values that are not updated are left unchanged. */
+	ZeroMissing						/** Existing constant values that are not updated are zeroed out. */
+};
 
 class CURVEEXPRESSION_API FEngine
 {
 public:
+	explicit FEngine(EParseFlags InParseFlags = EParseFlags::Default) :
+		ParseFlags(InParseFlags)
+	{ }
+
+	/** Initialize the engine with constants, making a copy of the value constant map. */
+	explicit FEngine(
+		const TMap<FName, float>& InConstants,
+		EParseFlags InParseFlags = EParseFlags::Default
+		);
+	
+	/** Initialize the engine with constants, taking ownership of the value constant map. */
+	explicit FEngine(
+		TMap<FName, float>&& InConstants,
+		EParseFlags InParseFlags = EParseFlags::Default
+		);
+
+	/** Update existing constant values. No new constants are added. */
+	void UpdateConstantValues(
+		const TMap<FName, float>& InConstants,
+		EUpdateConstantsMethod InUpdateMethod = EUpdateConstantsMethod::ZeroMissing
+		);
+
+	/** Return existing constants and their values. */
+	const TMap<FName, float>& GetConstantValues() const
+	{
+		return Constants;
+	}
+
 	/** Parse an expression and either receive an error or a expression execution object that
 	 *  matches this evaluator. */
 	TVariant<FExpressionObject, FParseError> Parse(
-		FStringView InExpression,
-		TFunction<TOptional<float>(FName InConstantName)> InConstantEvaluator = {} 
+		FStringView InExpression
 		) const;
 
-	/** Executes the given expression object.
-	 *  \param InExpressionObject The parsed expression object to evaluate.
-	 *  \param InConstantEvaluator A function that returns the value of a given constant, if it's known, or none otherwise.
-	 */
+	/** Executes the given expression object */
 	float Execute(
-		const FExpressionObject& InExpressionObject,
-		TFunctionRef<TOptional<float>(FName InConstantName)> InConstantEvaluator
+		const FExpressionObject& InExpressionObject
 		) const;
 	
 	/** Evaluate an expression directly and returns the result. If the expression is incorrectly
 	 *  formed, then the optional return value holds nothing.
 	 */ 
 	TOptional<float> Evaluate(
-		FStringView InExpression,
-		TFunctionRef<TOptional<float>(FName InConstantName)> InConstantEvaluator
+		FStringView InExpression
 		) const;
 
 	/** Verify an expression. Does not check for variable validity, only that the expression
 	 *  can be evaluated */
 	TOptional<FParseError> Verify(
-		FStringView InExpression,
-		TFunction<TOptional<float>(FName InConstantName)> InConstantEvaluator = {} 
+		FStringView InExpression
 		) const;
 	
 private:
@@ -160,42 +162,11 @@ private:
 		Max					// Must be last
 	};
 	using FToken = TVariant<EOperatorToken, /* Identifier */ FName, /* Value */ float>;
-
-	struct FTokenParseResult
-	{
-		FTokenParseResult() = default;
-		
-		FTokenParseResult(FToken&& InToken, FStringView InTokenRange, const TCHAR* InExpression) :
-			Result(TInPlaceType<FToken>(), MoveTemp(InToken)),
-			Location(InTokenRange, InExpression)
-		{}
-		
-		FTokenParseResult(FString&& InError, FStringView InTokenRange, const TCHAR* InExpression) :
-			Result(TInPlaceType<FString>(), MoveTemp(InError)),
-			Location(InTokenRange, InExpression)
-		{}
-		
-		bool IsError() const { return Result.IsType<FString>(); }
-		bool IsToken() const { return Result.IsType<FToken>(); }
-		
-		const FToken& GetToken() const
-		{
-			return Result.Get<FToken>();
-		}
-		
-		FParseError GetParseError()
-		{
-			return FParseError(MoveTemp(Result.Get<FString>()), Location);
-		}
-		
-		TVariant<FString, FToken> Result;
-		FParseLocation Location{};
-	};
 	
 	/** Parse a floating point value and return the result (or error). Same semantics
 	 *  apply as with ParseToken.
 	 */
-	static FTokenParseResult ParseFloat(
+	static TVariant<FToken, FParseError> ParseFloat(
 		FStringView& InOutParseRange,
 		FStringView InExpression
 		);
@@ -203,26 +174,39 @@ private:
 	/** Parse a constant reference and return the result (or error). Same semantics
 	 *  apply as with ParseToken.
 	 */
-	static FTokenParseResult ParseIdentifier(
+	TVariant<FToken, FParseError> ParseIdentifier(
 		FStringView& InOutParseRange,
 		FStringView InExpression
-		);
+		) const;
 	
 	/** Parse a single token from the beginning of the range, returning either the token
 	 *  or an error, if one was encountered.
 	 *  If the token was successfully parsed, then the given string view range gets updated
 	 *  with the new starting point at the end of the token just parsed
 	 */
-	static FTokenParseResult ParseToken(
+	TVariant<FToken, FParseError> ParseToken(
 		FStringView& InOutParseRange,
 		FStringView InExpression
-		);
+		) const;
 
+	/** Helper function to construct a token parse error result */
+	static TVariant<FToken, FParseError> ParseTokenError(
+		FString&& InError,
+		const TCHAR* InCurrentPos,
+		FStringView InExpression
+		);
+	
 	/** Helper function to get a text version of a token. */
 	static FString TokenToString(
 		const FToken& InToken
 		);
 	
+	/** The list of constants, stored using the FName key since the expression object refers
+	 *  to it directly */
+	// TBD: Add support for variables by allowing TFunction{Ref} objects along with direct
+	// constants.
+	TMap<FName, float> Constants;
+
 	/** For now we just have built-in functions */
 	friend struct FInitializeBuiltinFunctions;
 	using FunctionType = TFunction<float(TArrayView<const float>)>;
@@ -233,20 +217,9 @@ private:
 	};
 	static TMap<FName, int32> FunctionNameIndex;
 	static TArray<FFunctionInfo> Functions;
-};
-
-}
 
+	/** The parse flags for this engine. Only set once */
+	EParseFlags ParseFlags = EParseFlags::Default;
+};
 
-static FArchive& operator<<(FArchive& Ar, CurveExpression::Evaluator::FExpressionObject& InObject)
-{
-	if (Ar.IsLoading())
-	{
-		InObject.Load(Ar);
-	}
-	else
-	{
-		InObject.Save(Ar);
-	}
-	return Ar;
 }
