Curve Expression: Can't use PURE_VIRTUAL in BP structs.

#rnx

diff --git a/Engine/Plugins/Experimental/Animation/CurveExpression/Source/Runtime/Public/AnimNode_RemapCurvesBase.h b/Engine/Plugins/Experimental/Animation/CurveExpression/Source/Runtime/Public/AnimNode_RemapCurvesBase.h
--- a/Engine/Plugins/Experimental/Animation/CurveExpression/Source/Runtime/Public/AnimNode_RemapCurvesBase.h
+++ b/Engine/Plugins/Experimental/Animation/CurveExpression/Source/Runtime/Public/AnimNode_RemapCurvesBase.h
@@ -50,7 +50,6 @@ struct CURVEEXPRESSION_API FAnimNode_RemapCurvesBase :
 	virtual bool HasPreUpdate() const override { return true; }
 	virtual void PreUpdate(const UAnimInstance* InAnimInstance) override;
 	virtual void Update_AnyThread(const FAnimationUpdateContext& Context) override;
-	virtual void Evaluate_AnyThread(FPoseContext& Output) override PURE_VIRTUAL(FAnimNode_RemapCurvesBase::Evaluate_AnyThread, );
 	
 
 	// Parse and cache expressions defined in ExpressionList or CurveExpressions, depending on on use. 
