Curve Expression: Fix build break due to Perforce having a day.

#rnx

diff --git a/Engine/Plugins/Experimental/Animation/CurveExpression/Source/Runtime/Private/AnimNode_RemapCurves.cpp b/Engine/Plugins/Experimental/Animation/CurveExpression/Source/Runtime/Private/AnimNode_RemapCurves.cpp
--- a/Engine/Plugins/Experimental/Animation/CurveExpression/Source/Runtime/Private/AnimNode_RemapCurves.cpp
+++ b/Engine/Plugins/Experimental/Animation/CurveExpression/Source/Runtime/Private/AnimNode_RemapCurves.cpp
@@ -4,27 +4,6 @@
 
 #include "Animation/AnimInstanceProxy.h"
 
-void FAnimNode_RemapCurves::Initialize_AnyThread(const FAnimationInitializeContext& Context)
-{
-	FAnimNode_RemapCurvesBase::Initialize_AnyThread(Context);
-
-	const USkeleton* Skeleton = Context.AnimInstanceProxy->GetSkeleton();
-	if (ensure(Skeleton))
-	{
-		const FSmartNameMapping* CurveNameMapping = Skeleton->GetSmartNameContainer(USkeleton::AnimCurveMappingName);
-
-		CurveNameMapping->Iterate([this](const FSmartNameMappingIterator& Iterator)
-		{
-			if (FName CurveName; Iterator.GetName(CurveName))
-			{
-				CurveNameToUIDMap.Add(CurveName, Iterator.GetIndex());
-			}
-		});
-	}
-	
-	CurveEvalResult.Reserve(GetCompiledAssignments().Num());
-}
-
 void FAnimNode_RemapCurves::Evaluate_AnyThread(FPoseContext& Output)
 {
 	DECLARE_SCOPE_HIERARCHICAL_COUNTER_ANIMNODE(Evaluate_AnyThread)
@@ -34,31 +13,21 @@ void FAnimNode_RemapCurves::Evaluate_AnyThread(FPoseContext& Output)
 	FPoseContext SourceData(Output);
 	SourcePose.Evaluate(SourceData);
 	
-	CurveEvalResult.Reset();
-
+	FBlendedCurve Curve;
+	Curve.Reserve(GetCompiledAssignments().Num());
+	
 	for (const TTuple<FName, FExpressionObject>& Assignment: GetCompiledAssignments())
 	{
-		if (const SmartName::UID_Type* UID = CurveNameToUIDMap.Find(Assignment.Key))
-		{
-			CurveEvalResult.Add({*UID, FEngine().Execute(Assignment.Value,
-				[&CurveIDMap = CurveNameToUIDMap, &CurveValues = SourceData.Curve](const FName InName) -> TOptional<float>
-				{
-					if (const SmartName::UID_Type* SrcUID = CurveIDMap.Find(InName))
-					{
-						return CurveValues.Get(*SrcUID);
-					}
-					return {};
-				}
-			)});
-		}
+		Curve.Add(Assignment.Key, FEngine().Execute(Assignment.Value,
+			[&SourceCurve = SourceData.Curve](const FName InName) -> TOptional<float>
+			{
+				return SourceCurve.Get(InName);
+			}
+		));
 	}
 	
 	Output = SourceData;
-
-	for( const TTuple<SmartName::UID_Type, float>& CurveValue: CurveEvalResult)
-	{
-		Output.Curve.Set(CurveValue.Key, CurveValue.Value);
-	}
+	Output.Curve.Combine(Curve);
 }
 
 bool FAnimNode_RemapCurves::Serialize(FArchive& Ar)
