Curve Expression: Fix build break due to Perforce having a day. Take 2.

#rnx

diff --git a/Engine/Plugins/Experimental/Animation/CurveExpression/Source/Runtime/Public/AnimNode_RemapCurves.h b/Engine/Plugins/Experimental/Animation/CurveExpression/Source/Runtime/Public/AnimNode_RemapCurves.h
--- a/Engine/Plugins/Experimental/Animation/CurveExpression/Source/Runtime/Public/AnimNode_RemapCurves.h
+++ b/Engine/Plugins/Experimental/Animation/CurveExpression/Source/Runtime/Public/AnimNode_RemapCurves.h
@@ -1,6 +1,6 @@
 ﻿// Copyright Epic Games, Inc. All Rights Reserved.
-
 #pragma once
+
 #include "AnimNode_RemapCurvesBase.h"
 
 #include "AnimNode_RemapCurves.generated.h"
@@ -13,15 +13,9 @@ struct CURVEEXPRESSION_API FAnimNode_RemapCurves :
 	GENERATED_BODY()
 
 	// FAnimNode_Base interface
-	virtual void Initialize_AnyThread(const FAnimationInitializeContext& Context) override;
 	virtual void Evaluate_AnyThread(FPoseContext& Output) override;
 	
 	bool Serialize(FArchive& Ar);
-
-private:
-	// Used only by Initialize_AnyThread / Evaluate_AnyThread
-	TMap<FName, SmartName::UID_Type> CurveNameToUIDMap;
-	TArray<TTuple<SmartName::UID_Type, float>> CurveEvalResult;
 };
 
 template<> struct TStructOpsTypeTraits<FAnimNode_RemapCurves> : public TStructOpsTypeTraitsBase2<FAnimNode_RemapCurves>
