Curve Expressions: Fix verify expressions for target curves.

#rnx

diff --git a/Engine/Plugins/Experimental/Animation/CurveExpression/Source/Runtime/Private/AnimNode_RemapCurvesBase.cpp b/Engine/Plugins/Experimental/Animation/CurveExpression/Source/Runtime/Private/AnimNode_RemapCurvesBase.cpp
--- a/Engine/Plugins/Experimental/Animation/CurveExpression/Source/Runtime/Private/AnimNode_RemapCurvesBase.cpp
+++ b/Engine/Plugins/Experimental/Animation/CurveExpression/Source/Runtime/Private/AnimNode_RemapCurvesBase.cpp
@@ -222,11 +222,23 @@ void FAnimNode_RemapCurvesBase::VerifyExpressions(
 	
 	if (InTargetComponent->GetSkeletalMeshAsset() && InTargetComponent->GetSkeletalMeshAsset()->GetSkeleton())
 	{
-		TargetCurveNames = InTargetComponent->GetSkeletalMeshAsset()->GetSkeleton()->GetSmartNameContainer(USkeleton::AnimCurveMappingName);
+		const USkeleton* TargetSkeleton = InTargetComponent->GetSkeletalMeshAsset()->GetSkeleton();
+		TargetCurveNames = TargetSkeleton->GetSmartNameContainer(USkeleton::AnimCurveMappingName);
+		
+		if (!TargetCurveNames)
+		{
+			ReportAndLog(FString::Printf(TEXT("No target curves available for '%s'."), *TargetSkeleton->GetPackage()->GetPathName()));
+		}
 	}
 	if (InSourceComponent->GetSkeletalMeshAsset() && InSourceComponent->GetSkeletalMeshAsset()->GetSkeleton())
 	{
+		const USkeleton* SourceSkeleton = InSourceComponent->GetSkeletalMeshAsset()->GetSkeleton();
 		SourceCurveNames = InSourceComponent->GetSkeletalMeshAsset()->GetSkeleton()->GetSmartNameContainer(USkeleton::AnimCurveMappingName);
+		
+		if (!SourceCurveNames)
+		{
+			ReportAndLog(FString::Printf(TEXT("No source curves available for '%s'."), *SourceSkeleton->GetPackage()->GetPathName()));
+		}
 	}
 	*/
 
@@ -235,7 +247,9 @@ void FAnimNode_RemapCurvesBase::VerifyExpressions(
 		FEngine VerificationEngine;
 
 		/*
-		if (TargetCurveNames && TargetCurveNames->FindUID(ExpressionPair.Key) != SmartName::MaxUID)
+		// Only check for missing target curve names if we _do_ have a list of target curves, since we report the absence
+		// of the curve container above.
+		if (TargetCurveNames && TargetCurveNames->FindUID(ExpressionPair.Key) == SmartName::MaxUID)
 		{
 			ReportAndLog(FString::Printf(TEXT("Target curve '%s' does not exist."), *ExpressionPair.Key.ToString()));
 			bFoundError = true;
