Curve Expressions: Fix failure to change the source to Expression Map on load for older incarnation of the nodes.

#rnx

diff --git a/Engine/Plugins/Experimental/Animation/CurveExpression/Source/Runtime/Private/AnimNode_RemapCurvesBase.cpp b/Engine/Plugins/Experimental/Animation/CurveExpression/Source/Runtime/Private/AnimNode_RemapCurvesBase.cpp
--- a/Engine/Plugins/Experimental/Animation/CurveExpression/Source/Runtime/Private/AnimNode_RemapCurvesBase.cpp
+++ b/Engine/Plugins/Experimental/Animation/CurveExpression/Source/Runtime/Private/AnimNode_RemapCurvesBase.cpp
@@ -30,6 +30,7 @@ void FAnimNode_RemapCurvesBase::Initialize_AnyThread(
 	DECLARE_SCOPE_HIERARCHICAL_COUNTER_ANIMNODE(Initialize_AnyThread)
 	Super::Initialize_AnyThread(Context);
 	SourcePose.Initialize(Context);
+	GetEvaluateGraphExposedInputs().Execute(Context);
 	
 	// Make sure the expressions are parsed and ready to go.
 	ParseAndCacheExpressions();
@@ -46,18 +47,6 @@ void FAnimNode_RemapCurvesBase::CacheBones_AnyThread(
 }
 
 
-void FAnimNode_RemapCurvesBase::PreUpdate(const UAnimInstance* InAnimInstance)
-{
-	DECLARE_SCOPE_HIERARCHICAL_COUNTER_ANIMNODE(PreUpdate);
-	
-	// If we're using the expression map and the expressions are not constant, check if we need to reparse them here.
-	if (ExpressionSource == ERemapCurvesExpressionSource::ExpressionMap &&
-		(!bExpressionsImmutable || !ExpressionMapHash.IsSet() || GetTypeHash(CurveExpressions) != ExpressionMapHash.GetValue()))
-	{
-		ParseAndCacheExpressions();
-	}
-}
-
 void FAnimNode_RemapCurvesBase::Update_AnyThread(
 	const FAnimationUpdateContext& Context
 	)
@@ -69,6 +58,14 @@ void FAnimNode_RemapCurvesBase::Update_AnyThread(
 
 	// Evaluate any BP logic plugged into this node
 	GetEvaluateGraphExposedInputs().Execute(Context);
+
+	// If we have mutable expressions, check if they've changed here, and recompile if necessary. If the set of named 
+	// constants have changed, the new ones will not have the correct value until the next PreUpdate. 
+	if (ExpressionSource == ERemapCurvesExpressionSource::ExpressionMap && !bExpressionsImmutable &&
+		GetTypeHash(CurveExpressions) != ExpressionMapHash.GetValue())
+	{
+		ParseAndCacheExpressions();
+	}
 }
 
 
diff --git a/Engine/Plugins/Experimental/Animation/CurveExpression/Source/Runtime/Private/AnimNode_RemapCurvesFromMesh.cpp b/Engine/Plugins/Experimental/Animation/CurveExpression/Source/Runtime/Private/AnimNode_RemapCurvesFromMesh.cpp
--- a/Engine/Plugins/Experimental/Animation/CurveExpression/Source/Runtime/Private/AnimNode_RemapCurvesFromMesh.cpp
+++ b/Engine/Plugins/Experimental/Animation/CurveExpression/Source/Runtime/Private/AnimNode_RemapCurvesFromMesh.cpp
@@ -23,8 +23,6 @@ void FAnimNode_RemapCurvesFromMesh::PreUpdate(
 	// may have changed from underneath us.
 	RefreshMeshComponent(InAnimInstance->GetSkelMeshComponent());
 
-	FAnimNode_RemapCurvesBase::PreUpdate(InAnimInstance);
-	
 	const USkeletalMeshComponent* CurrentMeshComponent = CurrentlyUsedSourceMeshComponent.IsValid() ? CurrentlyUsedSourceMeshComponent.Get() : nullptr;
 
 	if (CurrentMeshComponent && CurrentMeshComponent->GetSkeletalMeshAsset() && CurrentMeshComponent->IsRegistered())
@@ -109,8 +107,8 @@ void FAnimNode_RemapCurvesFromMesh::GatherDebugData(
 bool FAnimNode_RemapCurvesFromMesh::Serialize(FArchive& Ar)
 {
 	SerializeNode(Ar, this, StaticStruct());
-
-	if (IsLoading() && Ar.CustomVer(FCurveExpressionCustomVersion::GUID) < FCurveExpressionCustomVersion::SerializedExpressions)
+	
+	if (Ar.IsLoading() && Ar.CustomVer(FCurveExpressionCustomVersion::GUID) < FCurveExpressionCustomVersion::SerializedExpressions)
 	{
 		// Default to expression map if loading from an older version.
 		ExpressionSource = ERemapCurvesExpressionSource::ExpressionMap;
diff --git a/Engine/Plugins/Experimental/Animation/CurveExpression/Source/Runtime/Public/AnimNode_RemapCurvesBase.h b/Engine/Plugins/Experimental/Animation/CurveExpression/Source/Runtime/Public/AnimNode_RemapCurvesBase.h
--- a/Engine/Plugins/Experimental/Animation/CurveExpression/Source/Runtime/Public/AnimNode_RemapCurvesBase.h
+++ b/Engine/Plugins/Experimental/Animation/CurveExpression/Source/Runtime/Public/AnimNode_RemapCurvesBase.h
@@ -26,7 +26,7 @@ struct CURVEEXPRESSION_API FAnimNode_RemapCurvesBase :
 
 	UPROPERTY(EditAnywhere, EditFixedSize, BlueprintReadWrite, Category = Links)
 	FPoseLink SourcePose;
-	
+
 	UPROPERTY(EditAnywhere, Category = Expressions)
 	ERemapCurvesExpressionSource ExpressionSource = ERemapCurvesExpressionSource::ExpressionList;
 
@@ -42,13 +42,11 @@ struct CURVEEXPRESSION_API FAnimNode_RemapCurvesBase :
 	/** The expression map given is immutable and will not change during runtime. Improves performance. */ 
 	UPROPERTY(EditAnywhere, BlueprintReadWrite, Category=Expressions, DisplayName = "Expression Map Does Not Change at Runtime", meta=(NeverAsPin, EditCondition="ExpressionSource==ERemapCurvesExpressionSource::ExpressionMap", EditConditionHides))
 	bool bExpressionsImmutable = true;
-
+	
 	
 	// FAnimNode_Base overrides
 	virtual void Initialize_AnyThread(const FAnimationInitializeContext& Context) override;
 	virtual void CacheBones_AnyThread(const FAnimationCacheBonesContext& Context) override;
-	virtual bool HasPreUpdate() const override { return true; }
-	virtual void PreUpdate(const UAnimInstance* InAnimInstance) override;
 	virtual void Update_AnyThread(const FAnimationUpdateContext& Context) override;
 	
 
diff --git a/Engine/Plugins/Experimental/Animation/CurveExpression/Source/Runtime/Public/AnimNode_RemapCurvesFromMesh.h b/Engine/Plugins/Experimental/Animation/CurveExpression/Source/Runtime/Public/AnimNode_RemapCurvesFromMesh.h
--- a/Engine/Plugins/Experimental/Animation/CurveExpression/Source/Runtime/Public/AnimNode_RemapCurvesFromMesh.h
+++ b/Engine/Plugins/Experimental/Animation/CurveExpression/Source/Runtime/Public/AnimNode_RemapCurvesFromMesh.h
@@ -22,6 +22,7 @@ struct CURVEEXPRESSION_API FAnimNode_RemapCurvesFromMesh :
 	bool bUseAttachedParent = false;
 	
 	// FAnimNode_Base interface
+	virtual bool HasPreUpdate() const override { return true; }
 	virtual void PreUpdate(const UAnimInstance* InAnimInstance) override;
 	virtual void Evaluate_AnyThread(FPoseContext& Output) override;
 	virtual void GatherDebugData(FNodeDebugData& DebugData) override;
