Curve Expressions: Use a per-asset shadow copy, rather than a per-thread shadow copy of the compiled expression data.

#rnx

diff --git a/Engine/Plugins/Experimental/Animation/CurveExpression/Source/Runtime/Private/CurveExpressionsDataAsset.cpp b/Engine/Plugins/Experimental/Animation/CurveExpression/Source/Runtime/Private/CurveExpressionsDataAsset.cpp
--- a/Engine/Plugins/Experimental/Animation/CurveExpression/Source/Runtime/Private/CurveExpressionsDataAsset.cpp
+++ b/Engine/Plugins/Experimental/Animation/CurveExpression/Source/Runtime/Private/CurveExpressionsDataAsset.cpp
@@ -50,27 +50,18 @@ TArray<FCurveExpressionParsedAssignment> FCurveExpressionList::GetParsedAssignme
 	return ParsedAssignments;
 }
 
-struct FCompiledExpressionScratchArea : public TThreadSingleton<FCompiledExpressionScratchArea>
+void UCurveExpressionsDataAsset::SynchronizeThreadCopy() const
 {
-	// Shadowed data copy for extraction from non-game threads.
-	uint32 SerialNumber{UINT_MAX};
-	TArray<FName> NamedConstants;
-	TMap<FName, CurveExpression::Evaluator::FExpressionObject> ExpressionMap;
-};
-
-
-void UCurveExpressionsDataAsset::SynchronizeThreadLocalData() const
-{
-	FCompiledExpressionScratchArea& ThreadLocalData = FCompiledExpressionScratchArea::Get();
-		
 	const uint32 Serial = SerialNumber.load();
-	if (ThreadLocalData.SerialNumber != Serial)
+	if (ThreadCopy.SerialNumber != Serial)
 	{
 		FScopeLock ScopeLock(&ExpressionMapWriteLock);
+		
+		FExpressionDataShadowCopy& MutableThreadCopy = const_cast<FExpressionDataShadowCopy&>(ThreadCopy);
 
-		ThreadLocalData.SerialNumber = Serial;
-		ThreadLocalData.NamedConstants = NamedConstants;
-		ThreadLocalData.ExpressionMap = ExpressionMap;
+		MutableThreadCopy.SerialNumber = Serial;
+		MutableThreadCopy.NamedConstants = NamedConstants;
+		MutableThreadCopy.ExpressionMap = ExpressionMap;
 	}
 }
 
@@ -83,8 +74,8 @@ const TArray<FName>& UCurveExpressionsDataAsset::GetCompiledExpressionConstants(
 	}
 	else
 	{
-		SynchronizeThreadLocalData();
-		return FCompiledExpressionScratchArea::Get().NamedConstants;
+		SynchronizeThreadCopy();
+		return ThreadCopy.NamedConstants;
 	}
 }
 
@@ -97,8 +88,8 @@ const TMap<FName, CurveExpression::Evaluator::FExpressionObject>& UCurveExpressi
 	}
 	else
 	{
-		SynchronizeThreadLocalData();
-		return FCompiledExpressionScratchArea::Get().ExpressionMap;
+		SynchronizeThreadCopy();
+		return ThreadCopy.ExpressionMap;
 	}
 }
 
diff --git a/Engine/Plugins/Experimental/Animation/CurveExpression/Source/Runtime/Public/CurveExpressionsDataAsset.h b/Engine/Plugins/Experimental/Animation/CurveExpression/Source/Runtime/Public/CurveExpressionsDataAsset.h
--- a/Engine/Plugins/Experimental/Animation/CurveExpression/Source/Runtime/Public/CurveExpressionsDataAsset.h
+++ b/Engine/Plugins/Experimental/Animation/CurveExpression/Source/Runtime/Public/CurveExpressionsDataAsset.h
@@ -71,7 +71,7 @@ private:
 	void CompileExpressions();
 #endif
 
-	void SynchronizeThreadLocalData() const;
+	void SynchronizeThreadCopy() const;
 	
 	mutable FCriticalSection ExpressionMapWriteLock;
 	
@@ -81,4 +81,11 @@ private:
 	TArray<FName> NamedConstants;
 	
 	TMap<FName, CurveExpression::Evaluator::FExpressionObject> ExpressionMap;
+
+	const struct FExpressionDataShadowCopy
+	{
+		uint32 SerialNumber{UINT_MAX};
+		TArray<FName> NamedConstants;
+		TMap<FName, CurveExpression::Evaluator::FExpressionObject> ExpressionMap;
+	} ThreadCopy;
 };
