Curve Expression: Guard against FArchive state that is neither loading, saving, collecting references or counting memory.

#rnx

diff --git a/Engine/Plugins/Experimental/Animation/CurveExpression/Source/Runtime/Private/AnimNode_RemapCurvesBase.cpp b/Engine/Plugins/Experimental/Animation/CurveExpression/Source/Runtime/Private/AnimNode_RemapCurvesBase.cpp
--- a/Engine/Plugins/Experimental/Animation/CurveExpression/Source/Runtime/Private/AnimNode_RemapCurvesBase.cpp
+++ b/Engine/Plugins/Experimental/Animation/CurveExpression/Source/Runtime/Private/AnimNode_RemapCurvesBase.cpp
@@ -75,14 +75,17 @@ void FAnimNode_RemapCurvesBase::SerializeNode(
 	UScriptStruct* InNodeStruct
 	)
 {
-	Ar.UsingCustomVersion(FCurveExpressionCustomVersion::GUID);
+	if (Ar.IsLoading() || Ar.IsSaving() || Ar.IsCountingMemory() || Ar.IsObjectReferenceCollector())
+	{
+		Ar.UsingCustomVersion(FCurveExpressionCustomVersion::GUID);
 	
-	InNodeStruct->SerializeTaggedProperties(Ar, static_cast<uint8*>(InNodeThisPtr), InNodeStruct, nullptr);
+		InNodeStruct->SerializeTaggedProperties(Ar, static_cast<uint8*>(InNodeThisPtr), InNodeStruct, nullptr);
 
-	if (Ar.CustomVer(FCurveExpressionCustomVersion::GUID) >= FCurveExpressionCustomVersion::SerializedExpressions)
-	{
-		Ar << CachedExpressions;
-		Ar << ExpressionMapHash;
+		if (Ar.CustomVer(FCurveExpressionCustomVersion::GUID) >= FCurveExpressionCustomVersion::SerializedExpressions)
+		{
+			Ar << CachedExpressions;
+			Ar << ExpressionMapHash;
+		}
 	}
 }
 
