Curve Expressions: Add the ability to remove curves using the `undef()` operator (or a "-CurveName") shorthand.

diff --git a/Engine/Plugins/Experimental/Animation/CurveExpression/Source/Runtime/Private/AnimNode_RemapCurves.cpp b/Engine/Plugins/Experimental/Animation/CurveExpression/Source/Runtime/Private/AnimNode_RemapCurves.cpp
--- a/Engine/Plugins/Experimental/Animation/CurveExpression/Source/Runtime/Private/AnimNode_RemapCurves.cpp
+++ b/Engine/Plugins/Experimental/Animation/CurveExpression/Source/Runtime/Private/AnimNode_RemapCurves.cpp
@@ -16,20 +16,32 @@ void FAnimNode_RemapCurves::Evaluate_AnyThread(FPoseContext& Output)
 	SourcePose.Evaluate(SourceData);
 	
 	FBlendedCurve Curve;
+	Curve.CopyFrom(SourceData.Curve);
 	Curve.Reserve(GetCompiledAssignments().Num());
 	
 	for (const TTuple<FName, FExpressionObject>& Assignment: GetCompiledAssignments())
 	{
-		Curve.Add(Assignment.Key, FEngine().Execute(Assignment.Value,
+		const float Value = FEngine().Execute(Assignment.Value,
 			[&SourceCurve = SourceData.Curve](const FName InName) -> TOptional<float>
 			{
 				return SourceCurve.Get(InName);
 			}
-		));
+		);
+		
+		// If the value is valid, set the curve's value. If the value is NaN, remove the curve, since it's a signal
+		// for removal from the expression (i.e. `undef()` was used).
+		if (!FMath::IsNaN(Value))
+		{
+			Curve.Set(Assignment.Key, Value);
+		}
+		else
+		{
+			Curve.InvalidateCurveWeight(Assignment.Key);
+		}
 	}
 	
 	Output = SourceData;
-	Output.Curve.Combine(Curve);
+	Output.Curve.MoveFrom(Curve);
 }
 
 bool FAnimNode_RemapCurves::Serialize(FArchive& Ar)
diff --git a/Engine/Plugins/Experimental/Animation/CurveExpression/Source/Runtime/Private/AnimNode_RemapCurvesFromMesh.cpp b/Engine/Plugins/Experimental/Animation/CurveExpression/Source/Runtime/Private/AnimNode_RemapCurvesFromMesh.cpp
--- a/Engine/Plugins/Experimental/Animation/CurveExpression/Source/Runtime/Private/AnimNode_RemapCurvesFromMesh.cpp
+++ b/Engine/Plugins/Experimental/Animation/CurveExpression/Source/Runtime/Private/AnimNode_RemapCurvesFromMesh.cpp
@@ -75,11 +75,9 @@ void FAnimNode_RemapCurvesFromMesh::Evaluate_AnyThread(
 	
 	Output = SourceData;
 
-	FBlendedCurve Curve;
-	Curve.Reserve(GetCompiledAssignments().Num());	
 	for (const TTuple<FName, FExpressionObject>& Assignment: GetCompiledAssignments())
 	{
-		Curve.Add(Assignment.Key, FEngine().Execute(Assignment.Value,
+		const float Value = FEngine().Execute(Assignment.Value,
 			[&CurveValues = SourceCurveValues](const FName InName) -> TOptional<float>
 			{
 				if (const float* Value = CurveValues.Find(InName))
@@ -87,10 +85,19 @@ void FAnimNode_RemapCurvesFromMesh::Evaluate_AnyThread(
 					return *Value;
 				}
 				return {};
-			}
-		));
+			});
+
+		// If the value is valid, set the curve's value. If the value is NaN, remove the curve, since it's a signal
+		// for removal from the expression (i.e. `undef()` was used).
+		if (!FMath::IsNaN(Value))
+		{
+			Output.Curve.Set(Assignment.Key, Value);
+		}
+		else
+		{
+			Output.Curve.InvalidateCurveWeight(Assignment.Key);
+		}
 	}
-	Output.Curve.Combine(Curve);
 }
 
 
diff --git a/Engine/Plugins/Experimental/Animation/CurveExpression/Source/Runtime/Private/CurveExpressionsDataAsset.cpp b/Engine/Plugins/Experimental/Animation/CurveExpression/Source/Runtime/Private/CurveExpressionsDataAsset.cpp
--- a/Engine/Plugins/Experimental/Animation/CurveExpression/Source/Runtime/Private/CurveExpressionsDataAsset.cpp
+++ b/Engine/Plugins/Experimental/Animation/CurveExpression/Source/Runtime/Private/CurveExpressionsDataAsset.cpp
@@ -13,24 +13,34 @@
 
 TArray<FCurveExpressionAssignment> FCurveExpressionList::GetAssignments() const
 {
+	// A helper string to mark an expression as returning an undefined value (for curve removal).
+	static FString UndefExpr(TEXT("undef()"));
+	
 	TArray<FCurveExpressionAssignment> ExpressionAssignments;
 	int32 LineIndex = 0;
 	
 	UE::String::ParseLines(AssignmentExpressions,
 		[&ExpressionAssignments, &LineIndex](FStringView InLine)
-	{
-		int32 AssignmentPos;
-		if (InLine.FindChar('=', AssignmentPos))
 		{
-			const FStringView AssignmentTarget = InLine.Left(AssignmentPos).TrimStartAndEnd();
-			const FStringView SourceExpression = InLine.Mid(AssignmentPos + 1).TrimStartAndEnd();
-			if (!AssignmentTarget.IsEmpty() && !SourceExpression.IsEmpty())
+			InLine.TrimStartAndEndInline();
+			
+			int32 AssignmentPos;
+			if (InLine.FindChar('=', AssignmentPos))
 			{
-				ExpressionAssignments.Add({LineIndex, FName(AssignmentTarget), SourceExpression});
+				const FStringView AssignmentTarget = InLine.Left(AssignmentPos).TrimStartAndEnd();
+				const FStringView SourceExpression = InLine.Mid(AssignmentPos + 1).TrimStartAndEnd();
+				if (!AssignmentTarget.IsEmpty() && !SourceExpression.IsEmpty())
+				{
+					ExpressionAssignments.Add({LineIndex, FName(AssignmentTarget), SourceExpression});
+				}
 			}
-		}
-		LineIndex++;
-	});
+			else if (InLine.StartsWith('-'))
+			{
+				const FStringView AssignmentTarget = InLine.Mid(1).TrimStartAndEnd();
+				ExpressionAssignments.Add({LineIndex, FName(AssignmentTarget), UndefExpr});
+			}
+			LineIndex++;
+		});
 	return ExpressionAssignments;
 }
 
diff --git a/Engine/Plugins/Experimental/Animation/CurveExpression/Source/Runtime/Private/ExpressionEvaluator.cpp b/Engine/Plugins/Experimental/Animation/CurveExpression/Source/Runtime/Private/ExpressionEvaluator.cpp
--- a/Engine/Plugins/Experimental/Animation/CurveExpression/Source/Runtime/Private/ExpressionEvaluator.cpp
+++ b/Engine/Plugins/Experimental/Animation/CurveExpression/Source/Runtime/Private/ExpressionEvaluator.cpp
@@ -62,8 +62,11 @@ struct FInitializeBuiltinFunctions
 		/** pi() */
 		CE_EXPR(pi, 0, UE_PI)
 		
-		/** pi() */
+		/** e() */
 		CE_EXPR(e, 0, UE_EULERS_NUMBER)
+
+		/** undef() */
+		CE_EXPR(undef, 0, std::numeric_limits<float>::signaling_NaN());
 	}
 } GInitializeBuiltinFunctions;
 #undef CE_EXPR
