Curve Expressions: Harden against possible initialization order issues. Fail gracefully if a function index is not found (either due to missing functions or corrupted expression bytestream).

#rnx

diff --git a/Engine/Plugins/Animation/DeformerGraph/Source/OptimusCore/Private/OptimusExpressionEvaluator.cpp b/Engine/Plugins/Animation/DeformerGraph/Source/OptimusCore/Private/OptimusExpressionEvaluator.cpp
--- a/Engine/Plugins/Animation/DeformerGraph/Source/OptimusCore/Private/OptimusExpressionEvaluator.cpp
+++ b/Engine/Plugins/Animation/DeformerGraph/Source/OptimusCore/Private/OptimusExpressionEvaluator.cpp
@@ -9,16 +9,13 @@ namespace Optimus::Expression
 {
 
 #define CE_EXPR(_N_, _A_, _E_) { \
-	FEngine::FunctionNameIndex.Add(FName(#_N_), FEngine::Functions.Num()); \
-	FEngine::Functions.Add({_A_, [](TArrayView<const float> V) -> float { return _E_; } }); \
+	FunctionNameIndex.Add(FName(#_N_), Functions.Num()); \
+	Functions.Add({_A_, [](TArrayView<const float> V) -> float { return _E_; } }); \
 	}  
 
-TArray<FEngine::FFunctionInfo> FEngine::Functions;
-TMap<FName, int32> FEngine::FunctionNameIndex;
-
-struct FInitializeBuiltinFunctions
+static struct FBuiltinFunctions
 {
-	FInitializeBuiltinFunctions()
+	FBuiltinFunctions()
 	{
 		// clamp(value, min, max)
 		CE_EXPR(clamp, 3, FMath::Clamp(V[0], V[1], V[2]))
@@ -65,7 +62,31 @@ struct FInitializeBuiltinFunctions
 		/** pi() */
 		CE_EXPR(e, 0, UE_EULERS_NUMBER)
 	}
-} GInitializeBuiltinFunctions;
+
+	int32 FindByName(const FName InName) const
+	{
+		if (const int32* Index = FunctionNameIndex.Find(InName))
+		{
+			return *Index;
+		}
+		return INDEX_NONE;
+	}
+
+	bool IsValidFunctionIndex(const int32 InIndex) const
+	{
+		return Functions.IsValidIndex(InIndex);
+	}
+	
+	const FEngine::FFunctionInfo& GetInfoByIndex(const int32 InIndex) const
+	{
+		return Functions[InIndex];
+	}
+	
+private:
+	TMap<FName, int32> FunctionNameIndex;
+	TArray<FEngine::FFunctionInfo> Functions;
+	
+} GBuiltinFunctions;
 #undef CE_EXPR
 
 
@@ -723,15 +744,15 @@ TVariant<FExpressionObject, FParseError> FEngine::Parse(
 				if (!Expression.IsEmpty() && Expression.Top().IsType<FName>())
 				{
 					FName ConstantName = Expression.Top().Get<FName>();
-					const int32 *FunctionIndexPtr = FunctionNameIndex.Find(ConstantName);
-					if (FunctionIndexPtr == nullptr)
+					const int32 FunctionIndex = GBuiltinFunctions.FindByName(ConstantName);
+					if (FunctionIndex == INDEX_NONE)
 					{
 						return ParseError(FString::Printf(TEXT("Unknown function '%s'"), *ConstantName.ToString()), Locations.Top());
 					}
 
 					Expression.Pop();
 					const FParseLocation LastLocation = Locations.Pop();
-					FunctionStack.Push({*FunctionIndexPtr, 0, OperatorStack.Num(), Expression.Num(), LastLocation});
+					FunctionStack.Push({FunctionIndex, 0, OperatorStack.Num(), Expression.Num(), LastLocation});
 				}
 			}
 			else if (Op == EOperatorToken::ParenClose)
@@ -749,7 +770,7 @@ TVariant<FExpressionObject, FParseError> FEngine::Parse(
 				{
 					// Compute the number of arguments.
 					const FFunctionCallInfo CallInfo = FunctionStack.Pop();
-					const FFunctionInfo& FunctionInfo = Functions[CallInfo.FunctionIndex];
+					const FFunctionInfo& FunctionInfo = GBuiltinFunctions.GetInfoByIndex(CallInfo.FunctionIndex);
 					const int32 ArgumentCount = CallInfo.CountedCommas + (CallInfo.ExpressionSize != Expression.Num());
 
 					if (ArgumentCount != FunctionInfo.ArgumentCount)
@@ -935,7 +956,14 @@ float FEngine::Execute(
 		}
 		else if (const FExpressionObject::FFunctionRef* FuncRef = Token.TryGet<FExpressionObject::FFunctionRef>())
 		{
-			const FFunctionInfo& FunctionInfo = Functions[FuncRef->Index];
+			// Technically, this shouldn't be necessary, but if the expression stream is broken, or from the wrong version,
+			// exit early to avoid a crash.
+			if (!GBuiltinFunctions.IsValidFunctionIndex(FuncRef->Index))
+			{
+				return 0.0f;
+			}
+				
+			const FFunctionInfo& FunctionInfo = GBuiltinFunctions.GetInfoByIndex(FuncRef->Index);
 			check(FunctionInfo.ArgumentCount <= ValueStack.Num());
 
 			TArrayView<float> ValueView(ValueStack.GetData() + ValueStack.Num() - FunctionInfo.ArgumentCount, FunctionInfo.ArgumentCount);
diff --git a/Engine/Plugins/Animation/DeformerGraph/Source/OptimusCore/Public/OptimusExpressionEvaluator.h b/Engine/Plugins/Animation/DeformerGraph/Source/OptimusCore/Public/OptimusExpressionEvaluator.h
--- a/Engine/Plugins/Animation/DeformerGraph/Source/OptimusCore/Public/OptimusExpressionEvaluator.h
+++ b/Engine/Plugins/Animation/DeformerGraph/Source/OptimusCore/Public/OptimusExpressionEvaluator.h
@@ -224,15 +224,13 @@ private:
 		);
 	
 	/** For now we just have built-in functions */
-	friend struct FInitializeBuiltinFunctions;
+	friend struct FBuiltinFunctions;
 	using FunctionType = TFunction<float(TArrayView<const float>)>;
 	struct FFunctionInfo
 	{
 		int32 ArgumentCount;
 		FunctionType FunctionPtr;
 	};
-	static TMap<FName, int32> FunctionNameIndex;
-	static TArray<FFunctionInfo> Functions;
 };
 
 }
diff --git a/Engine/Plugins/Experimental/Animation/CurveExpression/Source/Runtime/Private/AnimNode_RemapCurvesBase.cpp b/Engine/Plugins/Experimental/Animation/CurveExpression/Source/Runtime/Private/AnimNode_RemapCurvesBase.cpp
--- a/Engine/Plugins/Experimental/Animation/CurveExpression/Source/Runtime/Private/AnimNode_RemapCurvesBase.cpp
+++ b/Engine/Plugins/Experimental/Animation/CurveExpression/Source/Runtime/Private/AnimNode_RemapCurvesBase.cpp
@@ -209,7 +209,7 @@ void FAnimNode_RemapCurvesBase::VerifyExpressions(
 		UE_LOG(LogCurveExpression, Warning, TEXT("%s"), *InMessage);
 	};
 	
-	if (CurveExpressions.IsEmpty())
+	if (GetRawExpressions().IsEmpty())
 	{
 		ReportAndLog(TEXT("No curve expressions set."));
 		return;
diff --git a/Engine/Plugins/Experimental/Animation/CurveExpression/Source/Runtime/Private/ExpressionEvaluator.cpp b/Engine/Plugins/Experimental/Animation/CurveExpression/Source/Runtime/Private/ExpressionEvaluator.cpp
--- a/Engine/Plugins/Experimental/Animation/CurveExpression/Source/Runtime/Private/ExpressionEvaluator.cpp
+++ b/Engine/Plugins/Experimental/Animation/CurveExpression/Source/Runtime/Private/ExpressionEvaluator.cpp
@@ -9,16 +9,13 @@ namespace CurveExpression::Evaluator
 {
 
 #define CE_EXPR(_N_, _A_, _E_) { \
-	FEngine::FunctionNameIndex.Add(FName(#_N_), FEngine::Functions.Num()); \
-	FEngine::Functions.Add({_A_, [](TArrayView<const float> V) -> float { return _E_; } }); \
+	FunctionNameIndex.Add(FName(#_N_), Functions.Num()); \
+	Functions.Add({_A_, [](TArrayView<const float> V) -> float { return _E_; } }); \
 	}  
 
-TArray<FEngine::FFunctionInfo> FEngine::Functions;
-TMap<FName, int32> FEngine::FunctionNameIndex;
-
-struct FInitializeBuiltinFunctions
+static struct FBuiltinFunctions
 {
-	FInitializeBuiltinFunctions()
+	FBuiltinFunctions()
 	{
 		// clamp(value, min, max)
 		CE_EXPR(clamp, 3, FMath::Clamp(V[0], V[1], V[2]))
@@ -68,7 +65,31 @@ struct FInitializeBuiltinFunctions
 		/** undef() */
 		CE_EXPR(undef, 0, std::numeric_limits<float>::signaling_NaN());
 	}
-} GInitializeBuiltinFunctions;
+
+	int32 FindByName(const FName InName) const
+	{
+		if (const int32* Index = FunctionNameIndex.Find(InName))
+		{
+			return *Index;
+		}
+		return INDEX_NONE;
+	}
+
+	bool IsValidFunctionIndex(const int32 InIndex) const
+	{
+		return Functions.IsValidIndex(InIndex);
+	}
+	
+	const FEngine::FFunctionInfo& GetInfoByIndex(const int32 InIndex) const
+	{
+		return Functions[InIndex];
+	}
+	
+private:
+	TMap<FName, int32> FunctionNameIndex;
+	TArray<FEngine::FFunctionInfo> Functions;
+	
+} GBuiltinFunctions;
 #undef CE_EXPR
 
 
@@ -726,15 +747,15 @@ TVariant<FExpressionObject, FParseError> FEngine::Parse(
 				if (!Expression.IsEmpty() && Expression.Top().IsType<FName>())
 				{
 					FName ConstantName = Expression.Top().Get<FName>();
-					const int32 *FunctionIndexPtr = FunctionNameIndex.Find(ConstantName);
-					if (FunctionIndexPtr == nullptr)
+					const int32 FunctionIndex = GBuiltinFunctions.FindByName(ConstantName);
+					if (FunctionIndex == INDEX_NONE)
 					{
 						return ParseError(FString::Printf(TEXT("Unknown function '%s'"), *ConstantName.ToString()), Locations.Top());
 					}
 
 					Expression.Pop();
 					const FParseLocation LastLocation = Locations.Pop();
-					FunctionStack.Push({*FunctionIndexPtr, 0, OperatorStack.Num(), Expression.Num(), LastLocation});
+					FunctionStack.Push({FunctionIndex, 0, OperatorStack.Num(), Expression.Num(), LastLocation});
 				}
 			}
 			else if (Op == EOperatorToken::ParenClose)
@@ -752,7 +773,7 @@ TVariant<FExpressionObject, FParseError> FEngine::Parse(
 				{
 					// Compute the number of arguments.
 					const FFunctionCallInfo CallInfo = FunctionStack.Pop();
-					const FFunctionInfo& FunctionInfo = Functions[CallInfo.FunctionIndex];
+					const FFunctionInfo& FunctionInfo = GBuiltinFunctions.GetInfoByIndex(CallInfo.FunctionIndex);
 					const int32 ArgumentCount = CallInfo.CountedCommas + (CallInfo.ExpressionSize != Expression.Num());
 
 					if (ArgumentCount != FunctionInfo.ArgumentCount)
@@ -938,7 +959,14 @@ float FEngine::Execute(
 		}
 		else if (const FExpressionObject::FFunctionRef* FuncRef = Token.TryGet<FExpressionObject::FFunctionRef>())
 		{
-			const FFunctionInfo& FunctionInfo = Functions[FuncRef->Index];
+			// Technically, this shouldn't be necessary, but if the expression stream is broken, or from the wrong version,
+			// exit early to avoid a crash.
+			if (!GBuiltinFunctions.IsValidFunctionIndex(FuncRef->Index))
+			{
+				return 0.0f;
+			}
+				
+			const FFunctionInfo& FunctionInfo = GBuiltinFunctions.GetInfoByIndex(FuncRef->Index);
 			check(FunctionInfo.ArgumentCount <= ValueStack.Num());
 
 			TArrayView<float> ValueView(ValueStack.GetData() + ValueStack.Num() - FunctionInfo.ArgumentCount, FunctionInfo.ArgumentCount);
diff --git a/Engine/Plugins/Experimental/Animation/CurveExpression/Source/Runtime/Public/ExpressionEvaluator.h b/Engine/Plugins/Experimental/Animation/CurveExpression/Source/Runtime/Public/ExpressionEvaluator.h
--- a/Engine/Plugins/Experimental/Animation/CurveExpression/Source/Runtime/Public/ExpressionEvaluator.h
+++ b/Engine/Plugins/Experimental/Animation/CurveExpression/Source/Runtime/Public/ExpressionEvaluator.h
@@ -224,15 +224,13 @@ private:
 		);
 	
 	/** For now we just have built-in functions */
-	friend struct FInitializeBuiltinFunctions;
+	friend struct FBuiltinFunctions;
 	using FunctionType = TFunction<float(TArrayView<const float>)>;
 	struct FFunctionInfo
 	{
 		int32 ArgumentCount;
 		FunctionType FunctionPtr;
 	};
-	static TMap<FName, int32> FunctionNameIndex;
-	static TArray<FFunctionInfo> Functions;
 };
 
 }
