Curve Expressions: Handle multi-threaded switch-outs of data assets.

#rnx

diff --git a/Engine/Plugins/Experimental/Animation/CurveExpression/Source/Runtime/Private/AnimNode_RemapCurvesBase.cpp b/Engine/Plugins/Experimental/Animation/CurveExpression/Source/Runtime/Private/AnimNode_RemapCurvesBase.cpp
--- a/Engine/Plugins/Experimental/Animation/CurveExpression/Source/Runtime/Private/AnimNode_RemapCurvesBase.cpp
+++ b/Engine/Plugins/Experimental/Animation/CurveExpression/Source/Runtime/Private/AnimNode_RemapCurvesBase.cpp
@@ -104,7 +104,12 @@ const TMap<FName, CurveExpression::Evaluator::FExpressionObject>& FAnimNode_Rema
 	case ERemapCurvesExpressionSource::DataAsset:
 		if (CurveExpressionsDataAsset)
 		{
-			return CurveExpressionsDataAsset->GetCompiledExpressionMap();
+			CachedAssetExpressionData = CurveExpressionsDataAsset->GetCompiledExpressionData();
+			
+			if (CachedAssetExpressionData)
+			{
+				return CachedAssetExpressionData->ExpressionMap;
+			}
 		}
 	}
 
@@ -126,7 +131,12 @@ const TArray<FName>& FAnimNode_RemapCurvesBase::GetCompiledExpressionConstants()
 	case ERemapCurvesExpressionSource::DataAsset:
 		if (CurveExpressionsDataAsset)
 		{
-			return CurveExpressionsDataAsset->GetCompiledExpressionConstants();
+			CachedAssetExpressionData = CurveExpressionsDataAsset->GetCompiledExpressionData();
+			
+			if (CachedAssetExpressionData)
+			{
+				return CachedAssetExpressionData->NamedConstants;
+			}
 		}
 	}
 
diff --git a/Engine/Plugins/Experimental/Animation/CurveExpression/Source/Runtime/Private/CurveExpressionCustomVersion.h b/Engine/Plugins/Experimental/Animation/CurveExpression/Source/Runtime/Private/CurveExpressionCustomVersion.h
--- a/Engine/Plugins/Experimental/Animation/CurveExpression/Source/Runtime/Private/CurveExpressionCustomVersion.h
+++ b/Engine/Plugins/Experimental/Animation/CurveExpression/Source/Runtime/Private/CurveExpressionCustomVersion.h
@@ -15,6 +15,7 @@ struct FCurveExpressionCustomVersion
 
 		// Serialized expressions
 		SerializedExpressions,
+		ExpressionDataInSharedObject,
 		
 		VersionPlusOne,
 		LatestVersion = VersionPlusOne - 1,
diff --git a/Engine/Plugins/Experimental/Animation/CurveExpression/Source/Runtime/Private/CurveExpressionsDataAsset.cpp b/Engine/Plugins/Experimental/Animation/CurveExpression/Source/Runtime/Private/CurveExpressionsDataAsset.cpp
--- a/Engine/Plugins/Experimental/Animation/CurveExpression/Source/Runtime/Private/CurveExpressionsDataAsset.cpp
+++ b/Engine/Plugins/Experimental/Animation/CurveExpression/Source/Runtime/Private/CurveExpressionsDataAsset.cpp
@@ -6,6 +6,7 @@
 #include "ExpressionEvaluator.h"
 #endif
 
+#include "CurveExpressionCustomVersion.h"
 #include "String/ParseLines.h"
 
 #include UE_INLINE_GENERATED_CPP_BY_NAME(CurveExpressionsDataAsset)
@@ -60,52 +61,10 @@ TArray<FCurveExpressionParsedAssignment> FCurveExpressionList::GetParsedAssignme
 	return ParsedAssignments;
 }
 
-void UCurveExpressionsDataAsset::SynchronizeThreadCopy() const
-{
-	const uint32 Serial = SerialNumber.load();
-	if (ThreadCopy.SerialNumber != Serial)
-	{
-		FScopeLock ScopeLock(&ExpressionMapWriteLock);
-		
-		FExpressionDataShadowCopy& MutableThreadCopy = const_cast<FExpressionDataShadowCopy&>(ThreadCopy);
-
-		MutableThreadCopy.SerialNumber = Serial;
-		MutableThreadCopy.NamedConstants = NamedConstants;
-		MutableThreadCopy.ExpressionMap = ExpressionMap;
-	}
-}
-
-
-const TArray<FName>& UCurveExpressionsDataAsset::GetCompiledExpressionConstants() const
-{
-	if (IsInGameThread())
-	{
-		return NamedConstants;
-	}
-	else
-	{
-		SynchronizeThreadCopy();
-		return ThreadCopy.NamedConstants;
-	}
-}
-
-
-const TMap<FName, CurveExpression::Evaluator::FExpressionObject>& UCurveExpressionsDataAsset::GetCompiledExpressionMap() const
-{
-	if (IsInGameThread())
-	{
-		return ExpressionMap;
-	}
-	else
-	{
-		SynchronizeThreadCopy();
-		return ThreadCopy.ExpressionMap;
-	}
-}
-
-
 void UCurveExpressionsDataAsset::Serialize(FArchive& Ar)
 {
+	Ar.UsingCustomVersion(FCurveExpressionCustomVersion::GUID);
+	
 	UObject::Serialize(Ar);
 
 #if WITH_EDITOR
@@ -116,8 +75,22 @@ void UCurveExpressionsDataAsset::Serialize(FArchive& Ar)
 	}
 #endif
 
+	if (Ar.IsLoading() || !ExpressionData.IsValid())
+	{
+		ExpressionData = MakeShared<FExpressionData>();
+	}
+
+	if (Ar.IsLoading() && Ar.CustomVer(FCurveExpressionCustomVersion::GUID) < FCurveExpressionCustomVersion::ExpressionDataInSharedObject)
+	{
+		ExpressionData->NamedConstants = NamedConstants_DEPRECATED;
+	}
+	else
+	{
+		Ar << ExpressionData->NamedConstants;
+	}
+	
 	// Serialize the compiled map so we can read it in cooked builds.
-	Ar << ExpressionMap;
+	Ar << ExpressionData->ExpressionMap;
 }
 
 
@@ -137,22 +110,24 @@ void UCurveExpressionsDataAsset::CompileExpressions()
 {
 	using namespace CurveExpression::Evaluator;
 	
-	FScopeLock ScopeLock(&ExpressionMapWriteLock);
-
+	TSharedPtr<FExpressionData> NewExpressionData(new FExpressionData);
+	
 	TSet<FName> ConstantNames;
 	
-	ExpressionMap.Reset();
 	for (const FCurveExpressionParsedAssignment& Assignment: Expressions.GetParsedAssignments())
 	{
 		if (const FExpressionObject* Expression = Assignment.Result.TryGet<FExpressionObject>())
 		{
 			ConstantNames.Append(Expression->GetUsedConstants());
-			ExpressionMap.Add(Assignment.TargetName, *Expression);
+			NewExpressionData->ExpressionMap.Add(Assignment.TargetName, *Expression);
 		}
 	}
 	
-	NamedConstants = ConstantNames.Array();
-	SerialNumber += 1;
+	NewExpressionData->NamedConstants = ConstantNames.Array();
+
+	// Swap the new set of expressions in. Any client referring to the old expression data will still be
+	// able to do so.
+	ExpressionData = NewExpressionData;
 }
 
 #endif
diff --git a/Engine/Plugins/Experimental/Animation/CurveExpression/Source/Runtime/Public/AnimNode_RemapCurvesBase.h b/Engine/Plugins/Experimental/Animation/CurveExpression/Source/Runtime/Public/AnimNode_RemapCurvesBase.h
--- a/Engine/Plugins/Experimental/Animation/CurveExpression/Source/Runtime/Public/AnimNode_RemapCurvesBase.h
+++ b/Engine/Plugins/Experimental/Animation/CurveExpression/Source/Runtime/Public/AnimNode_RemapCurvesBase.h
@@ -79,6 +79,8 @@ private:
 	// Serialized data.
 	UPROPERTY()
 	TArray<FName> CachedConstantNames;
+
+	mutable TSharedPtr<const FExpressionData> CachedAssetExpressionData;
 	
 	TMap<FName, CurveExpression::Evaluator::FExpressionObject> CachedExpressions;
 	TOptional<int32> ExpressionMapHash;
diff --git a/Engine/Plugins/Experimental/Animation/CurveExpression/Source/Runtime/Public/CurveExpressionsDataAsset.h b/Engine/Plugins/Experimental/Animation/CurveExpression/Source/Runtime/Public/CurveExpressionsDataAsset.h
--- a/Engine/Plugins/Experimental/Animation/CurveExpression/Source/Runtime/Public/CurveExpressionsDataAsset.h
+++ b/Engine/Plugins/Experimental/Animation/CurveExpression/Source/Runtime/Public/CurveExpressionsDataAsset.h
@@ -41,6 +41,12 @@ struct FCurveExpressionList
 	TArray<FCurveExpressionParsedAssignment> GetParsedAssignments() const;
 };
 
+struct FExpressionData
+{
+	TArray<FName> NamedConstants;
+	TMap<FName, CurveExpression::Evaluator::FExpressionObject> ExpressionMap;
+};
+
 
 UCLASS(BlueprintType)
 class CURVEEXPRESSION_API UCurveExpressionsDataAsset :
@@ -54,11 +60,7 @@ public:
 	FCurveExpressionList Expressions;
 #endif
 
-	/** Returns the list of named constants used by ExpressionList */
-	const TArray<FName>& GetCompiledExpressionConstants() const;
-
-	/** Returns the list of compiled expressions generated from ExpressionList */
-	const TMap<FName, CurveExpression::Evaluator::FExpressionObject>& GetCompiledExpressionMap() const;
+	TSharedPtr<const FExpressionData> GetCompiledExpressionData() const { return ExpressionData; }
 
 	// UObject overrides
 	void Serialize(FArchive& Ar) override;
@@ -70,22 +72,9 @@ private:
 #if WITH_EDITOR
 	void CompileExpressions();
 #endif
-
-	void SynchronizeThreadCopy() const;
-	
-	mutable FCriticalSection ExpressionMapWriteLock;
-	
-	std::atomic<uint32> SerialNumber{0};
 	
 	UPROPERTY()
-	TArray<FName> NamedConstants;
+	TArray<FName> NamedConstants_DEPRECATED;	
 	
-	TMap<FName, CurveExpression::Evaluator::FExpressionObject> ExpressionMap;
-
-	const struct FExpressionDataShadowCopy
-	{
-		uint32 SerialNumber{UINT_MAX};
-		TArray<FName> NamedConstants;
-		TMap<FName, CurveExpression::Evaluator::FExpressionObject> ExpressionMap;
-	} ThreadCopy;
+	TSharedPtr<FExpressionData> ExpressionData;
 };
