Curve Expression: Add a remap(v, min_in, max_in, min_out, max_out) function.
- This remaps the input value such that its value within the in_min/max range is linearly mapped to the out_min/out_max range, much like the FMath::GetMappedRangeValueClamped C++ function.

#rnx
#Synced-CL 48407888

diff --git a/Engine/Plugins/Experimental/Animation/CurveExpression/Source/Runtime/Private/ExpressionEvaluator.cpp b/Engine/Plugins/Experimental/Animation/CurveExpression/Source/Runtime/Private/ExpressionEvaluator.cpp
--- a/Engine/Plugins/Experimental/Animation/CurveExpression/Source/Runtime/Private/ExpressionEvaluator.cpp
+++ b/Engine/Plugins/Experimental/Animation/CurveExpression/Source/Runtime/Private/ExpressionEvaluator.cpp
@@ -4,6 +4,7 @@
 #include <limits> // IWYU pragma: keep
 
 #include "CurveExpressionModule.h"
+#include "Math/Range.h"
 
 namespace CurveExpression::Evaluator
 {
@@ -15,6 +16,7 @@ namespace CurveExpression::Evaluator
 
 static struct FBuiltinFunctions
 {
+	// New functions should only be appended to this list. The function index is serialized.
 	FBuiltinFunctions()
 	{
 		// clamp(value, min, max)
@@ -64,6 +66,9 @@ static struct FBuiltinFunctions
 
 		/** undef() */
 		CE_EXPR(undef, 0, std::numeric_limits<float>::signaling_NaN());
+		
+		/** remap(value, min_in, max_in, min_out, max_out */
+		CE_EXPR(remap, 5, FMath::GetMappedRangeValueClamped(FFloatRange(V[1], V[2]), FFloatRange(V[3], V[4]), V[0]) )
 	}
 
 	int32 FindByName(const FName InName) const
